ADDED GeometryCache files
ADDED Support for importing GeometryCache assets from Alembic files
ADDED CalculateRawMeshTangets function that only calculates the tangents without building the RawMesh

diff --git a/Engine/Source/Developer/MeshUtilities/Private/MeshUtilities.cpp b/Engine/Source/Developer/MeshUtilities/Private/MeshUtilities.cpp
--- a/Engine/Source/Developer/MeshUtilities/Private/MeshUtilities.cpp
+++ b/Engine/Source/Developer/MeshUtilities/Private/MeshUtilities.cpp
@@ -153,6 +153,8 @@ private:
 		TArray<int32>& OutGlobalMaterialIndices
 		) const;
 
+	virtual void CalculateRawMeshTangents(FRawMesh& OutRawMesh, const FMeshBuildSettings& BuildSettings);
+
 	// Need to call some members from this class, (which is internal to this module)
 	friend class FStaticMeshUtilityBuilder;
 };
@@ -4093,6 +4095,59 @@ bool FMeshUtilities::ConstructRawMesh(
 	return true;
 }
 
+void FMeshUtilities::CalculateRawMeshTangents(FRawMesh& OutRawMesh, const FMeshBuildSettings& BuildSettings)
+{
+	int32 NumWedges = OutRawMesh.WedgeIndices.Num();
+
+	// Figure out if we should recompute normals and tangents. By default generated LODs should not recompute normals
+	bool bRecomputeNormals = (BuildSettings.bRecomputeNormals) || OutRawMesh.WedgeTangentZ.Num() == 0;
+	bool bRecomputeTangents = (BuildSettings.bRecomputeTangents) || OutRawMesh.WedgeTangentX.Num() == 0 || OutRawMesh.WedgeTangentY.Num() == 0;
+
+	// Dump normals and tangents if we are recomputing them.
+	if (bRecomputeTangents)
+	{
+		OutRawMesh.WedgeTangentX.Empty(NumWedges);
+		OutRawMesh.WedgeTangentX.AddZeroed(NumWedges);
+		OutRawMesh.WedgeTangentY.Empty(NumWedges);
+		OutRawMesh.WedgeTangentY.AddZeroed(NumWedges);
+	}
+
+	if (bRecomputeNormals)
+	{
+		OutRawMesh.WedgeTangentZ.Empty(NumWedges);
+		OutRawMesh.WedgeTangentZ.AddZeroed(NumWedges);
+	}
+
+	// Compute any missing tangents.
+	if (bRecomputeNormals || bRecomputeTangents)
+	{
+		float ComparisonThreshold = GetComparisonThreshold(BuildSettings);
+		TMultiMap<int32, int32> OverlappingCorners;
+		FindOverlappingCorners(OverlappingCorners, OutRawMesh, ComparisonThreshold);
+
+		// Static meshes always blend normals of overlapping corners.
+		uint32 TangentOptions = ETangentOptions::BlendOverlappingNormals;
+		if (BuildSettings.bRemoveDegenerates)
+		{
+			// If removing degenerate triangles, ignore them when computing tangents.
+			TangentOptions |= ETangentOptions::IgnoreDegenerateTriangles;
+		}
+		if (BuildSettings.bUseMikkTSpace)
+		{
+			ComputeTangents_MikkTSpace(OutRawMesh, OverlappingCorners, TangentOptions);
+		}
+		else
+		{
+			ComputeTangents(OutRawMesh, OverlappingCorners, TangentOptions);
+		}
+	}
+
+	// At this point the mesh will have valid tangents.
+	check(OutRawMesh.WedgeTangentX.Num() == NumWedges);
+	check(OutRawMesh.WedgeTangentY.Num() == NumWedges);
+	check(OutRawMesh.WedgeTangentZ.Num() == NumWedges);
+}
+
 /*------------------------------------------------------------------------------
 	Mesh merging 
 ------------------------------------------------------------------------------*/
diff --git a/Engine/Source/Developer/MeshUtilities/Public/MeshUtilities.h b/Engine/Source/Developer/MeshUtilities/Public/MeshUtilities.h
--- a/Engine/Source/Developer/MeshUtilities/Public/MeshUtilities.h
+++ b/Engine/Source/Developer/MeshUtilities/Public/MeshUtilities.h
@@ -208,4 +208,13 @@ public:
 		TArray<UObject*>& OutAssetsToSync,
 		FVector& OutProxyLocation
 		) = 0;
+
+	/**
+	* CalculateRawMeshTangents
+	*
+	* @param OutRawMesh - InOut raw Mesh for calculating tangents
+	* @return void
+	*/
+	virtual void CalculateRawMeshTangents(FRawMesh& OutRawMesh, const FMeshBuildSettings& BuildSettings) = 0;
+
 };
diff --git a/Engine/Source/Runtime/GeometryCache/Classes/GeometryCache.h b/Engine/Source/Runtime/GeometryCache/Classes/GeometryCache.h
new file mode 100644
--- /dev/null
+++ b/Engine/Source/Runtime/GeometryCache/Classes/GeometryCache.h
@@ -0,0 +1,68 @@
+// Copyright 1998-2015 Epic Games, Inc. All Rights Reserved.
+#pragma once
+#include "Interfaces/Interface_AssetUserData.h"
+
+#include "GeometryCacheTrack.h"
+
+#include "GeometryCache.generated.h"
+
+/**
+* A Geometry Cache is a piece/set of geometry that consists of individual Mesh/Transformation samples.
+* In contrast with Static Meshes they can have their vertices animated in certain ways. * 
+*/
+UCLASS(collapsecategories, hidecategories = Object, customconstructor, BlueprintType, config = Engine)
+class GEOMETRYCACHE_API UGeometryCache : public UObject, public IInterface_AssetUserData
+{
+	GENERATED_UCLASS_BODY()
+public:
+	/**
+	* Default constructor
+	*/
+	UGeometryCache(const FObjectInitializer& ObjectInitializer = FObjectInitializer::Get());
+
+	// Begin UObject interface.
+	virtual void Serialize(FArchive& Ar) override;
+	virtual FString GetDesc() override;
+	virtual SIZE_T GetResourceSize(EResourceSizeMode::Type Mode) override;
+	virtual void GetAssetRegistryTags(TArray<FAssetRegistryTag>& OutTags) const override;
+
+	virtual void BeginDestroy() override;
+	virtual bool IsReadyForFinishDestroy() override;
+#if WITH_EDITOR
+	virtual void PreEditChange(UProperty* PropertyAboutToChange) override;
+#endif // WITH_EDITOR
+	static void AddReferencedObjects(UObject* InThis, FReferenceCollector& Collector);
+	// End UObject interface.
+
+	/**
+	* AddTrack
+	*
+	* @param Track - GeometryCacheTrack instance that is a part of the GeometryCacheAsset
+	* @return void
+	*/
+	void AddTrack(UGeometryCacheTrack* Track);
+		
+	/** Clears all stored data so the reimporting step can fill the instance again*/
+	void ClearForReimporting();	
+
+#if WITH_EDITORONLY_DATA
+	/** Importing data and options used for this Geometry cache object*/
+	UPROPERTY(Category = ImportSettings, VisibleAnywhere, Instanced)
+	class UAssetImportData* AssetImportData;	
+
+	/** Information for thumbnail rendering */
+	UPROPERTY(VisibleAnywhere, Instanced, Category = Thumbnail)
+	class UThumbnailInfo* ThumbnailInfo;
+#endif
+	
+	/** GeometryCache track defining the samples/geometry data for this GeomCache instance */
+	UPROPERTY(VisibleAnywhere, Category=GeometryCache)
+	TArray<UGeometryCacheTrack*> Tracks;
+private:
+	/** Storing the number of Vertex and RigidBody track for asset meta tags*/
+	uint32 NumVertexAnimationTracks;
+	uint32 NumTransformAnimationTracks;
+
+	/** A fence which is used to keep track of the rendering thread releasing the static mesh resources. */
+	FRenderCommandFence ReleaseResourcesFence;
+};
diff --git a/Engine/Source/Runtime/GeometryCache/Classes/GeometryCacheActor.h b/Engine/Source/Runtime/GeometryCache/Classes/GeometryCacheActor.h
new file mode 100644
--- /dev/null
+++ b/Engine/Source/Runtime/GeometryCache/Classes/GeometryCacheActor.h
@@ -0,0 +1,28 @@
+// Copyright 1998-2015 Epic Games, Inc. All Rights Reserved.
+
+#pragma once
+
+#include "GameFramework/Actor.h"
+#include "GeometryCacheActor.generated.h"
+
+class UGeometryCacheComponent;
+
+/** GeometryCache actor, serves as a place-able actor for GeometryCache objects*/
+UCLASS(hidedropdown, MinimalAPI, notplaceable, NotBlueprintable)
+class AGeometryCacheActor : public AActor
+{
+	GENERATED_UCLASS_BODY()
+
+	// Begin AActor overrides.
+#if WITH_EDITOR
+		virtual bool GetReferencedContentObjects(TArray<UObject*>& Objects) const override;
+#endif // WITH_EDITOR
+	// End AActor overrides.
+
+	UPROPERTY(Category = GeometryCacheActor, VisibleAnywhere, BlueprintReadOnly, meta = (ExposeFunctionCategories = "Mesh,Rendering,Physics,Components|GeometryCache"))
+	UGeometryCacheComponent* GeometryCacheComponent;
+public:
+	/** Returns GeometryCacheComponent subobject **/
+	UFUNCTION(BlueprintCallable, Category = "Components|GeometryCache")
+	UGeometryCacheComponent* GetGeometryCacheComponent() const;	
+};
diff --git a/Engine/Source/Runtime/GeometryCache/Classes/GeometryCacheComponent.h b/Engine/Source/Runtime/GeometryCache/Classes/GeometryCacheComponent.h
new file mode 100644
--- /dev/null
+++ b/Engine/Source/Runtime/GeometryCache/Classes/GeometryCacheComponent.h
@@ -0,0 +1,222 @@
+// Copyright 1998-2015 Epic Games, Inc. All Rights Reserved.
+
+#pragma once
+
+#include "GeometryCacheModulePrivatePCH.h"
+#include "GeometryCache.h"
+#include "DynamicMeshBuilder.h"
+
+#include "GeometryCacheComponent.generated.h"
+
+class FGeometryCacheSceneProxy;
+class UGeometryCache;
+
+/** Stores the RenderData for each individual track */
+USTRUCT()
+struct FTrackRenderData
+{
+	GENERATED_USTRUCT_BODY()
+
+	~FTrackRenderData()
+	{
+		MeshData = NULL;
+	}
+	
+	/** Pointer to FGeometryCacheMeshData containing vertex-data, bounding box and batch-info*/
+	FGeometryCacheMeshData* MeshData;
+	/** World matrix used to render this specific track */
+	FMatrix WorldMatrix;
+	/** Index Buffer generated according to the vertex-data */
+	TArray<int32> IndexBuffer;
+
+	void Reset()
+	{
+		IndexBuffer.Empty();
+		MeshData = NULL;
+	}
+};
+
+/** GeometryCacheComponent, encapsulates a GeometryCache asset instance and implements functionality for rendering/and playback of GeometryCaches */
+UCLASS(ClassGroup = (Rendering, Common), hidecategories = (Object, LOD), meta = (BlueprintSpawnableComponent), Experimental, ClassGroup = Experimental)
+class GEOMETRYCACHE_API UGeometryCacheComponent : public UMeshComponent
+{
+	GENERATED_UCLASS_BODY()
+
+	/** Required for access to (protected) TrackSections */
+	friend FGeometryCacheSceneProxy;
+		
+	// Begin UObject Interface
+	virtual void BeginDestroy() override;
+#if WITH_EDITOR
+	virtual void PreEditUndo() override;
+	virtual void PostEditUndo() override;
+	virtual void PostEditChangeProperty(FPropertyChangedEvent& PropertyChangedEvent) override;
+#endif // WITH_EDITOR
+	// End UObject Interface
+
+	// Begin UActorComponent Interface.
+	virtual void OnRegister() override;
+	virtual void OnUnregister() override;
+	virtual void TickComponent(float DeltaTime, enum ELevelTick TickType, FActorComponentTickFunction *ThisTickFunction) override;
+	// End UActorComponent Interface.
+
+	// Begin USceneComponent Interface.
+	virtual FBoxSphereBounds CalcBounds(const FTransform& LocalToWorld) const override;
+	/** Update LocalBounds member from the local box of each section */
+	void UpdateLocalBounds();
+	// Begin USceneComponent Interface.	
+
+	// Begin UPrimitiveComponent Interface.
+	virtual FPrimitiveSceneProxy* CreateSceneProxy() override;
+	// End UPrimitiveComponent Interface.
+
+	// Begin UMeshComponent Interface.
+	virtual int32 GetNumMaterials() const override;
+	// End UMeshComponent Interface.
+
+	/**
+	* CreateTrackSection, Create/replace a track section.
+	*
+	* @param WorldMatrix - World Matrix for the TrackSection
+	* @param MeshData - MeshData for the TrackSection (containing vertex-data, batch-info and bounding box)
+	*/
+	void CreateTrackSection(int32 SectionIndex, const FMatrix& WorldMatrix, FGeometryCacheMeshData* MeshData);
+
+	/**
+	* UpdateTrackSectionMeshData, Update only the Mesh Data (Vertices) for a specific section
+	*
+	* @param SectionIndex - Index of the section we want to update.
+	* @param MeshData - New MeshData pointer for the section.
+	*/
+	void UpdateTrackSectionMeshData(int32 SectionIndex, FGeometryCacheMeshData* MeshData);
+
+	/**
+	* UpdateTrackSectionMatrixData, Update only the World Matrix for a specific section
+	*
+	* @param SectionIndex - Index of the section we want to update
+	* @param WorldMatrix - New world matrix for the section.
+	*/
+	void UpdateTrackSectionMatrixData(int32 SectionIndex, const FMatrix& WorldMatrix);	
+
+	/**
+	* UpdateTrackSectionVertexbuffer, Update only the Vertex Buffer for a specific section
+	*
+	* @param SectionIndex - Index of the section we want to update
+	* @param MeshData - MeshData pointer to update the VB with
+	*/
+	void UpdateTrackSectionVertexbuffer(int32 SectionIndex, FGeometryCacheMeshData* MeshData);
+
+	/**
+	* UpdateTrackSectionIndexbuffer, Update only the Index Buffer for a specific section
+	*
+	* @param SectionIndex - Index of the section we want to update
+	* @param Indices - Array with new indices
+	*/
+	void UpdateTrackSectionIndexbuffer(int32 SectionIndex, const TArray<int32>& Indices );
+
+	/**
+	* OnObjectReimported QQQ
+	*
+	* @param ImportedGeometryCache - Callback function to refresh section data and update sceneproxy
+	*/
+	void OnObjectReimported(UGeometryCache* ImportedGeometryCache);
+
+	/** Start playback of GeometryCache */
+	UFUNCTION(BlueprintCallable, Category = "Components|GeometryCache")
+	void Play();
+
+	/** Start playback of GeometryCache from the start */
+	UFUNCTION(BlueprintCallable, Category = "Components|GeometryCache")
+	void PlayFromStart();
+
+	/** Start playback of GeometryCache in reverse*/
+	UFUNCTION(BlueprintCallable, Category = "Components|GeometryCache")
+	void PlayReversed();
+	
+	/** Start playback of GeometryCache from the end and play in reverse */
+	UFUNCTION(BlueprintCallable, Category = "Components|GeometryCache")
+	void PlayReversedFromEnd();
+
+	/** Pause playback of GeometryCache */
+	UFUNCTION(BlueprintCallable, Category = "Components|GeometryCache")
+	void Pause();
+
+	/** Stop playback of GeometryCache */
+	UFUNCTION(BlueprintCallable, Category = "Components|GeometryCache")
+	void Stop();
+
+	/** Get whether this GeometryCache is playing or not. */
+	UFUNCTION(BlueprintCallable, Category = "Components|GeometryCache")
+	bool IsPlaying() const;
+
+	/** Get whether this GeometryCache is playing in reverse or not. */
+	UFUNCTION(BlueprintCallable, Category = "Components|GeometryCache")
+	bool IsPlayingReversed() const;
+
+	/** Get whether this GeometryCache is looping or not. */
+	UFUNCTION(BlueprintCallable, Category = "Components|GeometryCache")
+	bool IsLooping() const;
+
+	/** Set whether this GeometryCache is looping or not. */
+	UFUNCTION(BlueprintCallable, Category = "Components|GeometryCache")
+	void SetLooping( const bool bNewLooping);
+
+	/** Get whether this GeometryCache is playing or not. */
+	UFUNCTION(BlueprintCallable, Category = "Components|GeometryCache")
+	float GetPlaybackSpeed() const;
+
+	/** Get whether this GeometryCache is playing or not. */
+	UFUNCTION(BlueprintCallable, Category = "Components|GeometryCache")
+	void SetPlaybackSpeed(const float NewPlaybackSpeed);
+	
+	/** Geometry Cache instance referenced by the component */
+	UPROPERTY(EditAnywhere, BlueprintReadOnly, Category = GeometryCache)
+	UGeometryCache* GeometryCache;
+	
+	UPROPERTY(EditAnywhere, Category = GeometryCache)
+	bool bRunning;
+
+	UPROPERTY(EditAnywhere, Category = GeometryCache)
+	bool bLooping;
+
+	UPROPERTY(EditAnywhere, Category = GeometryCache)
+	float StartTimeOffset;
+
+	UPROPERTY(EditAnywhere, Category = GeometryCache)
+	float PlaybackSpeed;
+
+	UPROPERTY(VisibleAnywhere, Category = GeometryCache)
+	int32 NumTracks;
+
+	UPROPERTY(VisibleAnywhere, Category = GeometryCache)
+	float ElapsedTime;
+
+protected:
+	/**
+	* Invalidate both the Matrix and Mesh sample indices
+	*/
+	void InvalidateTrackSampleIndices();
+
+	/**
+	* ReleaseResources, clears and removes data stored/copied from GeometryCache instance	
+	*/
+	void ReleaseResources();	
+
+	/** Local space bounds of mesh */
+	FBoxSphereBounds LocalBounds;
+
+	/** Pointer to SceneProxy for the instance of this component */
+	FGeometryCacheSceneProxy* SceneProxy;
+
+	/** Matrix and Mesh sample index for each individual track */
+	TArray<int32> TrackMeshSampleIndices;
+	TArray<int32> TrackMatrixSampleIndices;
+	
+	/** Array containing the TrackData (used for rendering) for each individual track*/
+	TArray<FTrackRenderData> TrackSections;
+
+	/** Play (time) direction, either -1.0f or 1.0f */
+	float PlayDirection;
+	/** Duration of the animation (maximum time) */
+	float Duration;
+};
diff --git a/Engine/Source/Runtime/GeometryCache/Classes/GeometryCacheMeshData.h b/Engine/Source/Runtime/GeometryCache/Classes/GeometryCacheMeshData.h
new file mode 100644
--- /dev/null
+++ b/Engine/Source/Runtime/GeometryCache/Classes/GeometryCacheMeshData.h
@@ -0,0 +1,96 @@
+// Copyright 1998-2015 Epic Games, Inc. All Rights Reserved.
+#pragma once
+
+#include "DynamicMeshBuilder.h"
+
+#include "Box.h"
+#include "GeometryCacheMeshData.generated.h"
+
+/** Stores per-batch data used for rendering */
+USTRUCT()
+struct FGeometryCacheMeshBatchInfo
+{
+	GENERATED_USTRUCT_BODY()
+
+	/** Starting index into IndexBuffer to draw from */
+	uint32 StartIndex;
+	/** Total number of Triangles to draw */
+	uint32 NumTriangles;
+	/** Index to Material used to draw this batch */
+	uint32 MaterialIndex;
+
+	friend FArchive& operator<<(FArchive& Ar, FGeometryCacheMeshBatchInfo& Mesh)
+	{
+		Ar << Mesh.StartIndex;
+		Ar << Mesh.NumTriangles;
+		Ar << Mesh.MaterialIndex;
+		return Ar;
+	}
+};
+
+/** Stores per Track/Mesh data used for rendering*/
+USTRUCT()
+struct FGeometryCacheMeshData
+{
+	GENERATED_USTRUCT_BODY()
+
+	FGeometryCacheMeshData() {}
+	~FGeometryCacheMeshData()
+	{
+		Vertices.Empty();
+		BatchesInfo.Empty();
+		BoundingBox.Init();
+	}
+	
+	/** Draw-able vertices */
+	TArray<FDynamicMeshVertex> Vertices;
+	/** Array of per-batch info structs*/
+	TArray<FGeometryCacheMeshBatchInfo> BatchesInfo;
+	/** Bounding box for this sample in the track */
+	FBox BoundingBox;
+	
+	/** Serialization for FVertexAnimationSample. */
+	friend FArchive& operator<<(FArchive& Ar, FGeometryCacheMeshData& Mesh)
+	{
+		int32 NumVertices = 0;
+		
+		if (Ar.IsSaving())
+		{
+			NumVertices = Mesh.Vertices.Num();
+		}
+
+		Ar << NumVertices;
+		if (Ar.IsLoading())
+		{
+			Mesh.Vertices.AddUninitialized(NumVertices);
+		}
+
+		for (int32 VertexIndex = 0; VertexIndex < NumVertices; ++VertexIndex)
+		{
+			FDynamicMeshVertex& Vertex = Mesh.Vertices[VertexIndex];
+			Ar << Vertex.Position;
+			Ar << Vertex.TextureCoordinate;
+			Ar << Vertex.TangentX;
+			Ar << Vertex.TangentZ;
+			Ar << Vertex.Color;
+		}
+
+		Ar << Mesh.BoundingBox;
+		Ar << Mesh.BatchesInfo;
+
+		return Ar;
+	}
+
+	SIZE_T GetResourceSize() const
+	{
+		// Calculate resource size according to what is actually serialized
+		SIZE_T Size = 0;
+		Size += Vertices.Num() * sizeof(FDynamicMeshVertex);
+		Size += BatchesInfo.Num() * sizeof(FGeometryCacheMeshBatchInfo);
+		Size += sizeof(Vertices);
+		Size += sizeof(BatchesInfo);
+		Size += sizeof(BoundingBox);
+
+		return Size;
+	}
+};
diff --git a/Engine/Source/Runtime/GeometryCache/Classes/GeometryCacheTrack.h b/Engine/Source/Runtime/GeometryCache/Classes/GeometryCacheTrack.h
new file mode 100644
--- /dev/null
+++ b/Engine/Source/Runtime/GeometryCache/Classes/GeometryCacheTrack.h
@@ -0,0 +1,83 @@
+// Copyright 1998-2015 Epic Games, Inc. All Rights Reserved.
+#pragma once
+
+#include "GeometryCacheTrack.generated.h"
+
+struct FGeometryCacheMeshData;
+
+/** Base class for GeometryCache tracks, stores matrix animation data and implements functionality for it */
+UCLASS(collapsecategories, hidecategories = Object, MinimalAPI, BlueprintType, config = Engine)
+class UGeometryCacheTrack : public UObject
+{
+	GENERATED_UCLASS_BODY()
+
+	virtual ~UGeometryCacheTrack();
+	
+	// Begin UObject interface.
+	virtual void Serialize(FArchive& Ar) override;
+	virtual SIZE_T GetResourceSize(EResourceSizeMode::Type Mode) override;
+	// End UObject interface
+			
+	/**
+	* UpdateMatrixData
+	*
+	* @param Time - (Elapsed)Time to check against
+	* @param bLooping - Whether or not the animation should be played on a loop
+	* @param InOutMatrixSampleIndex - Hold the MatrixSampleIndex and will be updated if changed according to the Elapsed Time
+	* @param OutWorldMatrix - Will hold the new WorldMatrix if the SampleIndex changed
+	* @return const bool
+	*/
+	virtual const bool UpdateMatrixData(const float Time, const bool bLooping, int32& InOutMatrixSampleIndex, FMatrix& OutWorldMatrix);
+
+	/**
+	* UpdateMeshData
+	*
+	* * @param Time - (Elapsed)Time to check against
+	* @param bLooping - Whether or not the animation should be played on a loop
+	* @param InOutMeshSampleIndex - Hold the MeshSampleIndex and will be updated if changed according to the Elapsed Time
+	* @param OutVertices - Will hold the new VertexData if the SampleIndex changed
+	* @return const bool
+	*/
+	virtual const bool UpdateMeshData(const float Time, const bool bLooping, int32& InOutMeshSampleIndex, FGeometryCacheMeshData*& OutMeshData );
+
+	/**
+	* SetMatrixSamples, Set the Matrix animation Samples 
+	*
+	* @param Matrices - Array of Matrices
+	* @param SampleTimes - Array of SampleTimes
+	* @return void
+	*/
+	virtual void SetMatrixSamples(const TArray<FMatrix>& Matrices, const TArray<float>& SampleTimes);
+
+	/**
+	* GetMaxSampleTime, returns the time for the last sample can be considered as the total animation leght
+	*
+	* @return const float
+	*/
+	virtual const float GetMaxSampleTime() const;
+	
+	/**
+	* GetNumMaterials, total number of materials inside this track (depends on batches)
+	*
+	* @return const uint32
+	*/
+	const uint32 GetNumMaterials() const { return NumMaterials; }
+
+protected:
+	/**
+	* FindSampleIndexFromTime uses binary search to find the closest index to Time inside Samples
+	*
+	* @param SampleTimes - Array of Sample times used for the search
+	* @param Time - Time for which the closest index has to be found
+	* @param bLooping - Whether or not we should fmod Time according to the last entry in SampleTimes
+	* @return const uint32
+	*/
+	const uint32 FindSampleIndexFromTime(const TArray<float>& SampleTimes, const float Time, const bool bLooping);
+
+	/** Matrix sample data, both FMatrix and time*/
+	TArray<FMatrix> MatrixSamples;	
+	TArray<float> MatrixSampleTimes;
+
+	/** Number of materials for this track*/
+	uint32 NumMaterials;
+};
diff --git a/Engine/Source/Runtime/GeometryCache/Classes/GeometryCacheTrackFlipbookAnimation.h b/Engine/Source/Runtime/GeometryCache/Classes/GeometryCacheTrackFlipbookAnimation.h
new file mode 100644
--- /dev/null
+++ b/Engine/Source/Runtime/GeometryCache/Classes/GeometryCacheTrackFlipbookAnimation.h
@@ -0,0 +1,46 @@
+// Copyright 1998-2015 Epic Games, Inc. All Rights Reserved.
+#pragma once
+
+#include "GeometryCacheMeshData.h"
+
+#include "GeometryCacheTrackFlipbookAnimation.generated.h"
+
+/** Derived GeometryCacheTrack class, used for Transform animation. */
+UCLASS(collapsecategories, hidecategories = Object, BlueprintType, config = Engine)
+class GEOMETRYCACHE_API UGeometryCacheTrack_FlipbookAnimation : public UGeometryCacheTrack
+{
+	GENERATED_UCLASS_BODY()
+
+	virtual ~UGeometryCacheTrack_FlipbookAnimation();
+
+	// Begin UObject interface.
+	virtual SIZE_T GetResourceSize(EResourceSizeMode::Type Mode) override;
+	virtual void Serialize(FArchive& Ar) override;
+	virtual void BeginDestroy() override;
+	// End UObject interface.
+
+	// Begin UGeometryCacheTrack interface.
+	virtual const bool UpdateMeshData(const float Time, const bool bLooping, int32& InOutMeshSampleIndex, FGeometryCacheMeshData*& OutMeshData) override;
+	virtual const float GetMaxSampleTime() const override;
+	// End UGeometryCacheTrack interface.
+
+	/**
+	* Add a GeometryCacheMeshData sample to the Track
+	*
+	* @param MeshData - Holds the mesh data for the specific sample
+	* @param SampleTime - SampleTime for the specific sample being added
+	* @return void
+	*/
+	UFUNCTION()
+	void AddMeshSample(const FGeometryCacheMeshData& MeshData, const float SampleTime);
+
+private:
+	/** Number of Mesh Sample within this track */
+	UPROPERTY(VisibleAnywhere, Category = GeometryCache)
+	uint32 NumMeshSamples;
+
+	/** Stored data for each Mesh sample */
+	TArray<FGeometryCacheMeshData> MeshSamples;
+	TArray<float> MeshSampleTimes;
+};
+
diff --git a/Engine/Source/Runtime/GeometryCache/Classes/GeometryCacheTrackTransformAnimation.h b/Engine/Source/Runtime/GeometryCache/Classes/GeometryCacheTrackTransformAnimation.h
new file mode 100644
--- /dev/null
+++ b/Engine/Source/Runtime/GeometryCache/Classes/GeometryCacheTrackTransformAnimation.h
@@ -0,0 +1,32 @@
+// Copyright 1998-2015 Epic Games, Inc. All Rights Reserved.
+#pragma once
+
+#include "GeometryCacheMeshData.h"
+
+#include "GeometryCacheTrackTransformAnimation.generated.h"
+
+/** Derived GeometryCacheTrack class, used for Transform animation. */
+UCLASS(collapsecategories, hidecategories = Object, BlueprintType, config = Engine)
+class GEOMETRYCACHE_API UGeometryCacheTrack_TransformAnimation : public UGeometryCacheTrack
+{
+	GENERATED_UCLASS_BODY()
+		
+	// Begin UObject interface.
+	virtual SIZE_T GetResourceSize(EResourceSizeMode::Type Mode) override;
+	virtual void Serialize(FArchive& Ar) override;
+	// End UObject interface.
+	
+	// Begin UGeometryCacheTrack interface.
+	virtual const bool UpdateMeshData(const float Time, const bool bLooping, int32& InOutMeshSampleIndex, FGeometryCacheMeshData*& OutMeshData) override;
+	// End UGeometryCacheTrack interface.
+
+	/**
+	* Sets/updates the MeshData for this track
+	*
+	* @param NewMeshData - GeometryCacheMeshData instance later used as the rendered mesh	
+	*/
+	UFUNCTION()
+	void SetMesh(const FGeometryCacheMeshData& NewMeshData);	
+private:
+	FGeometryCacheMeshData MeshData;
+};
diff --git a/Engine/Source/Runtime/GeometryCache/GeometryCache.Build.cs b/Engine/Source/Runtime/GeometryCache/GeometryCache.Build.cs
new file mode 100644
--- /dev/null
+++ b/Engine/Source/Runtime/GeometryCache/GeometryCache.Build.cs
@@ -0,0 +1,29 @@
+// Copyright 1998-2015 Epic Games, Inc. All Rights Reserved.
+
+using UnrealBuildTool;
+
+public class GeometryCache : ModuleRules
+{
+	public GeometryCache(TargetInfo Target)
+	{
+        PublicIncludePaths.Add("Runtime/GeometryCache/Public");
+        PrivateIncludePaths.Add("Runtime/GeometryCache/Private");
+        
+        PublicDependencyModuleNames.AddRange(
+			new string[] {
+				"Core",
+				"CoreUObject", // @todo Mac: for some reason it's needed to link in debug on Mac
+				"Engine",
+				"Slate",
+				"SlateCore",
+                "InputCore",
+                "EditorStyle",
+                "UnrealEd",
+                "AssetTools",
+                "RenderCore",
+                "ShaderCore",
+                "RHI"
+			}
+		);
+	}
+}
diff --git a/Engine/Source/Runtime/GeometryCache/Private/ActorFactoryGeometryCache.cpp b/Engine/Source/Runtime/GeometryCache/Private/ActorFactoryGeometryCache.cpp
new file mode 100644
--- /dev/null
+++ b/Engine/Source/Runtime/GeometryCache/Private/ActorFactoryGeometryCache.cpp
@@ -0,0 +1,59 @@
+// Copyright 1998-2015 Epic Games, Inc. All Rights Reserved.
+
+#include "GeometryCacheModulePrivatePCH.h"
+#include "GeometryCache.h"
+#include "ActorFactoryGeometryCache.h"
+#include "GeometryCacheActor.h"
+#include "GeometryCacheComponent.h"
+
+UActorFactoryGeometryCache::UActorFactoryGeometryCache(const FObjectInitializer& ObjectInitializer)
+: Super(ObjectInitializer)
+{
+	DisplayName = FText::FromString( "Geometry Cache" );
+	NewActorClass = AGeometryCacheActor::StaticClass();
+	bUseSurfaceOrientation = true;
+}
+
+bool UActorFactoryGeometryCache::CanCreateActorFrom(const FAssetData& AssetData, FText& OutErrorMsg)
+{
+	if (!AssetData.IsValid() || !AssetData.GetClass()->IsChildOf(UGeometryCache::StaticClass()))
+	{
+		OutErrorMsg = FText::FromString("A valid GeometryCache must be specified.");
+		return false;
+	}
+
+	return true;
+}
+
+void UActorFactoryGeometryCache::PostSpawnActor(UObject* Asset, AActor* NewActor)
+{
+	Super::PostSpawnActor(Asset, NewActor);
+
+	UGeometryCache* GeometryCache = CastChecked<UGeometryCache>(Asset);
+
+	// Change properties
+	AGeometryCacheActor* GeometryCacheActor = CastChecked<AGeometryCacheActor>(NewActor);
+	UGeometryCacheComponent* GeometryCacheComponent = GeometryCacheActor->GetGeometryCacheComponent();
+	check(GeometryCacheComponent);
+
+	GeometryCacheComponent->UnregisterComponent();
+
+	// Set GeometryCache (data) instance
+	GeometryCacheComponent->GeometryCache = GeometryCache;
+
+	// Init Component
+	GeometryCacheComponent->RegisterComponent();
+}
+
+void UActorFactoryGeometryCache::PostCreateBlueprint(UObject* Asset, AActor* CDO)
+{
+	if (Asset != NULL && CDO != NULL)
+	{
+		// Set GeometryCache (data) instance
+		UGeometryCache* GeometryCache = CastChecked<UGeometryCache>(Asset);
+		AGeometryCacheActor* GeometryCacheActor = CastChecked<AGeometryCacheActor>(CDO);
+		UGeometryCacheComponent* GeometryCacheComponent = GeometryCacheActor->GetGeometryCacheComponent();
+
+		GeometryCacheComponent->GeometryCache = GeometryCache;
+	}
+}
\ No newline at end of file
diff --git a/Engine/Source/Runtime/GeometryCache/Private/ActorFactoryGeometryCache.h b/Engine/Source/Runtime/GeometryCache/Private/ActorFactoryGeometryCache.h
new file mode 100644
--- /dev/null
+++ b/Engine/Source/Runtime/GeometryCache/Private/ActorFactoryGeometryCache.h
@@ -0,0 +1,20 @@
+// Copyright 1998-2015 Epic Games, Inc. All Rights Reserved.
+
+#pragma once
+
+#include "GeometryCacheModulePrivatePCH.h"
+#include "ActorFactories/ActorFactory.h"
+#include "ActorFactoryGeometryCache.generated.h"
+
+/** Factory class for spawning and creating GeometryCacheActors */
+UCLASS(MinimalAPI, config = Editor)
+class UActorFactoryGeometryCache : public UActorFactory
+{
+	GENERATED_UCLASS_BODY()
+
+	// Begin UActorFactory Interface
+	virtual bool CanCreateActorFrom(const FAssetData& AssetData, FText& OutErrorMsg) override;
+	virtual void PostSpawnActor(UObject* Asset, AActor* NewActor) override;
+	virtual void PostCreateBlueprint(UObject* Asset, AActor* CDO) override;	
+	// End UActorFactory Interface
+};
diff --git a/Engine/Source/Runtime/GeometryCache/Private/AssetTypeActions_GeometryCache.cpp b/Engine/Source/Runtime/GeometryCache/Private/AssetTypeActions_GeometryCache.cpp
new file mode 100644
--- /dev/null
+++ b/Engine/Source/Runtime/GeometryCache/Private/AssetTypeActions_GeometryCache.cpp
@@ -0,0 +1,66 @@
+// Copyright 1998-2015 Epic Games, Inc. All Rights Reserved.
+
+#include "GeometryCacheModulePrivatePCH.h"
+#include "AssetTypeActions_GeometryCache.h"
+
+FText FAssetTypeActions_GeometryCache::GetName() const
+{
+	return NSLOCTEXT("AssetTypeActions", "AssetTypeActions_GeometryCache", "GeometryCache");
+}
+
+FColor FAssetTypeActions_GeometryCache::GetTypeColor() const
+{
+	return FColor(0, 255, 255);
+}
+
+UClass* FAssetTypeActions_GeometryCache::GetSupportedClass() const
+{
+	return UGeometryCache::StaticClass();
+}
+
+bool FAssetTypeActions_GeometryCache::HasActions(const TArray<UObject*>& InObjects) const
+{
+	return false;
+}
+
+void FAssetTypeActions_GeometryCache::GetActions(const TArray<UObject*>& InObjects, FMenuBuilder& MenuBuilder)
+{
+	FAssetTypeActions_Base::GetActions(InObjects, MenuBuilder);
+}
+
+void FAssetTypeActions_GeometryCache::OpenAssetEditor(const TArray<UObject*>& InObjects, TSharedPtr<class IToolkitHost> EditWithinLevelEditor /*= TSharedPtr<IToolkitHost>()*/)
+{
+	FAssetTypeActions_Base::OpenAssetEditor(InObjects, EditWithinLevelEditor);
+}
+
+uint32 FAssetTypeActions_GeometryCache::GetCategories()
+{
+	return EAssetTypeCategories::Basic;
+}
+
+class UThumbnailInfo* FAssetTypeActions_GeometryCache::GetThumbnailInfo(UObject* Asset) const
+{
+	UGeometryCache* GeometryCache = CastChecked<UGeometryCache>(Asset);
+	UThumbnailInfo* ThumbnailInfo = GeometryCache->ThumbnailInfo;
+	if (ThumbnailInfo == NULL)
+	{
+		ThumbnailInfo = NewObject<USceneThumbnailInfo>(GeometryCache);
+		GeometryCache->ThumbnailInfo = ThumbnailInfo;
+	}
+
+	return ThumbnailInfo;
+}
+
+bool FAssetTypeActions_GeometryCache::IsImportedAsset() const
+{
+	return true;
+}
+
+void FAssetTypeActions_GeometryCache::GetResolvedSourceFilePaths(const TArray<UObject*>& TypeAssets, TArray<FString>& OutSourceFilePaths) const
+{
+	for (auto& Asset : TypeAssets)
+	{
+		const auto GeometryCache = CastChecked<UGeometryCache>(Asset);
+		GeometryCache->AssetImportData->ExtractFilenames(OutSourceFilePaths);
+	}
+}
diff --git a/Engine/Source/Runtime/GeometryCache/Private/AssetTypeActions_GeometryCache.h b/Engine/Source/Runtime/GeometryCache/Private/AssetTypeActions_GeometryCache.h
new file mode 100644
--- /dev/null
+++ b/Engine/Source/Runtime/GeometryCache/Private/AssetTypeActions_GeometryCache.h
@@ -0,0 +1,23 @@
+// Copyright 1998-2015 Epic Games, Inc. All Rights Reserved.
+#pragma once
+
+#include "GeometryCache.h"
+#include "AssetTypeActions_Base.h"
+
+/** Type-actions class for GeometryCache assets*/
+class FAssetTypeActions_GeometryCache : public FAssetTypeActions_Base
+{
+public:
+	// IAssetTypeActions Implementation
+	virtual FText GetName() const override;
+	virtual FColor GetTypeColor() const override;
+	virtual UClass* GetSupportedClass() const override;
+	virtual bool HasActions(const TArray<UObject*>& InObjects) const override;
+	virtual void GetActions(const TArray<UObject*>& InObjects, FMenuBuilder& MenuBuilder) override;
+	virtual void OpenAssetEditor(const TArray<UObject*>& InObjects, TSharedPtr<class IToolkitHost> EditWithinLevelEditor = TSharedPtr<IToolkitHost>()) override;
+	virtual uint32 GetCategories() override;
+	virtual class UThumbnailInfo* GetThumbnailInfo(UObject* Asset) const override;
+	virtual bool IsImportedAsset() const override;
+	virtual void GetResolvedSourceFilePaths(const TArray<UObject*>& TypeAssets, TArray<FString>& OutSourceFilePaths) const override;
+	// End IAssetTypeActions
+};
diff --git a/Engine/Source/Runtime/GeometryCache/Private/GeometryCache.cpp b/Engine/Source/Runtime/GeometryCache/Private/GeometryCache.cpp
new file mode 100644
--- /dev/null
+++ b/Engine/Source/Runtime/GeometryCache/Private/GeometryCache.cpp
@@ -0,0 +1,118 @@
+// Copyright 1998-2015 Epic Games, Inc. All Rights Reserved.
+
+#include "GeometryCacheModulePrivatePCH.h"
+#include "GeometryCacheTrackTransformAnimation.h"
+#include "GeometryCacheTrackFlipbookAnimation.h"
+
+UGeometryCache::UGeometryCache(const FObjectInitializer& ObjectInitializer /*= FObjectInitializer::Get()*/) : UObject(ObjectInitializer)
+{
+	NumTransformAnimationTracks = 0;
+	NumVertexAnimationTracks = 0;
+}
+
+void UGeometryCache::Serialize(FArchive& Ar)
+{
+	// Custom serialization
+#if WITH_EDITORONLY_DATA
+	Ar << AssetImportData;
+#endif
+	Ar << Tracks;
+	Ar << NumVertexAnimationTracks;
+	Ar << NumTransformAnimationTracks;
+}
+
+FString UGeometryCache::GetDesc()
+{
+	const int32 NumTracks = Tracks.Num();
+	return FString("%d Tracks", NumTracks);
+}
+
+SIZE_T UGeometryCache::GetResourceSize(EResourceSizeMode::Type Mode)
+{
+	SIZE_T ResourceSize = 0;
+#if WITH_EDITORONLY_DATA
+	ResourceSize += sizeof( AssetImportData );
+#endif
+	// Calculate Resource Size according to what is serialized
+	const int32 NumTracks = Tracks.Num();
+	for (int32 TrackIndex = 0; TrackIndex < NumTracks; ++TrackIndex)
+	{
+		ResourceSize += Tracks[TrackIndex]->GetResourceSize(Mode);
+	}
+
+	ResourceSize += sizeof(Tracks);
+	ResourceSize += sizeof(NumVertexAnimationTracks);
+	ResourceSize += sizeof(NumTransformAnimationTracks);
+
+	return ResourceSize;
+}
+
+void UGeometryCache::GetAssetRegistryTags(TArray<FAssetRegistryTag>& OutTags) const
+{
+	// Information on number total of (per type) tracks
+	const int32 NumTracks = Tracks.Num();	
+	OutTags.Add(FAssetRegistryTag("Total Tracks", FString::FromInt(NumTracks), FAssetRegistryTag::TT_Numerical));
+	OutTags.Add(FAssetRegistryTag("Vertex Animation Tracks", FString::FromInt(NumVertexAnimationTracks), FAssetRegistryTag::TT_Numerical));
+	OutTags.Add(FAssetRegistryTag("Transform Animation Tracks", FString::FromInt(NumTransformAnimationTracks), FAssetRegistryTag::TT_Numerical));
+
+#if WITH_EDITORONLY_DATA
+	if (AssetImportData)
+	{
+		OutTags.Add(FAssetRegistryTag(SourceFileTagName(), AssetImportData->ToJson(), FAssetRegistryTag::TT_Hidden));
+	}
+#endif
+
+	Super::GetAssetRegistryTags(OutTags);
+}
+
+void UGeometryCache::BeginDestroy()
+{
+	Super::BeginDestroy();
+	NumVertexAnimationTracks = NumTransformAnimationTracks = 0;
+	Tracks.Empty();
+	ReleaseResourcesFence.BeginFence();
+}
+
+void UGeometryCache::ClearForReimporting()
+{
+	NumVertexAnimationTracks = NumTransformAnimationTracks = 0;
+	Tracks.Empty();
+
+	// Flush the resource release commands to the rendering thread to ensure that the edit change doesn't occur while a resource is still allocated
+	ReleaseResourcesFence.Wait();
+}
+
+bool UGeometryCache::IsReadyForFinishDestroy()
+{
+	return ReleaseResourcesFence.IsFenceComplete();
+}
+
+void UGeometryCache::PreEditChange(UProperty* PropertyAboutToChange)
+{
+	// Release the Geometry Cache resources.
+	NumVertexAnimationTracks = NumTransformAnimationTracks = 0;
+	Tracks.Empty();
+
+	// Flush the resource release commands to the rendering thread to ensure that the edit change doesn't occur while a resource is still allocated
+	ReleaseResourcesFence.Wait();
+}
+
+void UGeometryCache::AddReferencedObjects(UObject* InThis, FReferenceCollector& Collector)
+{
+	UGeometryCache* This = CastChecked<UGeometryCache>(InThis);	
+	Super::AddReferencedObjects(This, Collector);
+}
+
+void UGeometryCache::AddTrack(UGeometryCacheTrack* Track)
+{
+	if (Track->GetClass() == UGeometryCacheTrack_TransformAnimation::StaticClass())
+	{
+		NumTransformAnimationTracks++;
+	}
+	else if (Track->GetClass() == UGeometryCacheTrack_FlipbookAnimation::StaticClass())
+	{
+		NumVertexAnimationTracks++;
+	}
+
+	Tracks.Add(Track);
+}
diff --git a/Engine/Source/Runtime/GeometryCache/Private/GeometryCacheActor.cpp b/Engine/Source/Runtime/GeometryCache/Private/GeometryCacheActor.cpp
new file mode 100644
--- /dev/null
+++ b/Engine/Source/Runtime/GeometryCache/Private/GeometryCacheActor.cpp
@@ -0,0 +1,30 @@
+// Copyright 1998-2015 Epic Games, Inc. All Rights Reserved.
+
+#include "GeometryCacheModulePrivatePCH.h"
+#include "GeometryCacheActor.h"
+#include "UnrealEd.h"
+#include "GeometryCacheComponent.h"
+
+AGeometryCacheActor::AGeometryCacheActor(const FObjectInitializer& ObjectInitializer)
+: Super(ObjectInitializer)
+{
+	GeometryCacheComponent = CreateDefaultSubobject<UGeometryCacheComponent>(TEXT("GeometryCacheComponent"));
+	RootComponent = GeometryCacheComponent;
+}
+
+class UGeometryCacheComponent* AGeometryCacheActor::GetGeometryCacheComponent() const
+{
+	return GeometryCacheComponent;
+}
+
+bool AGeometryCacheActor::GetReferencedContentObjects(TArray<UObject*>& Objects) const
+{
+	Super::GetReferencedContentObjects(Objects);
+
+	if (GeometryCacheComponent && GeometryCacheComponent->GeometryCache)
+	{
+		Objects.Add(GeometryCacheComponent->GeometryCache);
+	}
+
+	return true;
+}
diff --git a/Engine/Source/Runtime/GeometryCache/Private/GeometryCacheComponent.cpp b/Engine/Source/Runtime/GeometryCache/Private/GeometryCacheComponent.cpp
new file mode 100644
--- /dev/null
+++ b/Engine/Source/Runtime/GeometryCache/Private/GeometryCacheComponent.cpp
@@ -0,0 +1,445 @@
+// Copyright 1998-2015 Epic Games, Inc. All Rights Reserved.
+
+#include "GeometryCacheModulePrivatePCH.h"
+#include "GeometryCacheComponent.h"
+
+#include "GeometryCacheSceneProxy.h"
+#include "GeometryCacheTrack.h"
+#include "GeometryCacheTrackFlipbookAnimation.h"
+
+DECLARE_CYCLE_STAT(TEXT("GeometryCacheTick"), STAT_GeometryCacheComponent_TickComponent, STATGROUP_GeometryCache);
+
+UGeometryCacheComponent::UGeometryCacheComponent(const FObjectInitializer& ObjectInitializer)
+: Super(ObjectInitializer)
+{
+	PrimaryComponentTick.bCanEverTick = true;
+	ElapsedTime = 0.0f;
+	bRunning = false;
+	bLooping = true;
+	PlayDirection = 1.0f;
+	StartTimeOffset = 0.0f;
+	PlaybackSpeed = 1.0f;
+	Duration = 0.0f;
+}
+
+void UGeometryCacheComponent::BeginDestroy()
+{
+	Super::BeginDestroy();
+	ReleaseResources();	
+}
+
+void UGeometryCacheComponent::OnRegister()
+{
+	NumTracks = 0;
+	TrackMatrixSampleIndices.Empty();
+	TrackMatrixSampleIndices.Empty();
+	TrackSections.Empty();
+
+	if (GeometryCache != NULL)
+	{
+		// Refresh NumTracks and clear Index Arrays
+		NumTracks = GeometryCache->Tracks.Num();		
+		TrackMeshSampleIndices.Empty(GeometryCache->Tracks.Num());
+		TrackMatrixSampleIndices.Empty(GeometryCache->Tracks.Num());
+
+		Duration = 0.0f;
+		// Create mesh sections for each GeometryCacheTrack
+		for (int32 TrackIndex = 0; TrackIndex < NumTracks; ++TrackIndex )
+		{ 				
+			UGeometryCacheTrack* Track = GeometryCache->Tracks[TrackIndex];			
+			FMatrix WorldMatrix;
+			int32 MeshSampleIndex = -1;
+			int32 MatrixSampleIndex = -1;
+			FGeometryCacheMeshData* MeshData = NULL;
+
+			// Retrieve the matrix/mesh data and the appropriate sample indices
+			Track->UpdateMatrixData(ElapsedTime + StartTimeOffset, bLooping, MatrixSampleIndex, WorldMatrix);
+			Track->UpdateMeshData(ElapsedTime + StartTimeOffset, bLooping, MeshSampleIndex, MeshData);
+
+			// First time so create rather than update the mesh sections
+			CreateTrackSection(TrackIndex, WorldMatrix, MeshData);
+
+			// Store the sample indices for both the mesh and matrix data
+			TrackMeshSampleIndices.Add(MeshSampleIndex);
+			TrackMatrixSampleIndices.Add(MatrixSampleIndex);
+
+			const float TrackMaxSampleTime = Track->GetMaxSampleTime();
+			Duration = ( Duration > TrackMaxSampleTime ) ? Duration : TrackMaxSampleTime;
+		}
+	}
+
+	Super::OnRegister();
+}
+
+void UGeometryCacheComponent::OnUnregister()
+{
+	Super::OnUnregister();
+
+	NumTracks = 0;
+	TrackMatrixSampleIndices.Empty();
+	TrackMatrixSampleIndices.Empty();
+	TrackSections.Empty();
+}
+
+void UGeometryCacheComponent::TickComponent(float DeltaTime, enum ELevelTick TickType, FActorComponentTickFunction *ThisTickFunction)
+{
+	SCOPE_CYCLE_COUNTER(STAT_GeometryCacheComponent_TickComponent);
+	if (GeometryCache && bRunning)
+	{
+		// Increase total elapsed time since BeginPlay according to PlayDirection and speed
+		ElapsedTime += (DeltaTime * PlayDirection * PlaybackSpeed);
+
+		// QQ change this for reverse playing? extra float for runningtime vs sample time?
+		if (ElapsedTime < 0.0f)
+		{
+			ElapsedTime += Duration;
+		}
+
+		for (int32 TrackIndex = 0; TrackIndex < NumTracks; ++TrackIndex)
+		{
+			// Determine if and which part of the section we have to update
+			UGeometryCacheTrack* Track = GeometryCache->Tracks[TrackIndex];
+			FMatrix WorldMatrix;
+			FGeometryCacheMeshData* MeshData = NULL;
+
+			const bool bUpdateMatrix = Track->UpdateMatrixData(ElapsedTime + StartTimeOffset, bLooping, TrackMatrixSampleIndices[TrackIndex], WorldMatrix);
+			const bool bUpdateMesh = Track->UpdateMeshData(ElapsedTime + StartTimeOffset, bLooping, TrackMeshSampleIndices[TrackIndex], MeshData);
+
+			// Update sections according what is required
+			if (bUpdateMatrix)
+			{
+				UpdateTrackSectionMatrixData(TrackIndex, WorldMatrix);
+			}
+
+			if (bUpdateMesh)
+			{
+				UpdateTrackSectionMeshData(TrackIndex, MeshData);
+			}			
+		}
+	}
+}
+
+FBoxSphereBounds UGeometryCacheComponent::CalcBounds(const FTransform& LocalToWorld) const
+{
+	return LocalBounds.TransformBy(LocalToWorld);
+}
+void UGeometryCacheComponent::UpdateLocalBounds()
+{
+	FBox LocalBox(0);
+
+	for (const FTrackRenderData& Section : TrackSections)
+	{
+		// Use World matrix per section for correct bounding box
+		LocalBox += (Section.MeshData->BoundingBox.TransformBy(Section.WorldMatrix));
+	}
+
+	LocalBounds = LocalBox.IsValid ? FBoxSphereBounds(LocalBox) : FBoxSphereBounds(FVector(0, 0, 0), FVector(0, 0, 0), 0); // fallback to reset box sphere bounds
+
+	UpdateBounds();
+}
+
+FPrimitiveSceneProxy* UGeometryCacheComponent::CreateSceneProxy()
+{
+	SceneProxy = new FGeometryCacheSceneProxy(this);
+	return SceneProxy;
+}
+
+void UGeometryCacheComponent::PostEditChangeProperty(FPropertyChangedEvent& PropertyChangedEvent)
+{
+	InvalidateTrackSampleIndices();
+	MarkRenderStateDirty();
+	Super::PostEditChangeProperty(PropertyChangedEvent);
+}
+
+int32 UGeometryCacheComponent::GetNumMaterials() const
+{
+	// Sum the number of materials per track
+	int32 TotalNumMaterial = 0;
+	if ( GeometryCache )
+	{
+		for (UGeometryCacheTrack* Track : GeometryCache->Tracks)
+		{
+			TotalNumMaterial += Track->GetNumMaterials();
+		}
+	}
+
+	return TotalNumMaterial;
+}
+
+void UGeometryCacheComponent::CreateTrackSection(int32 SectionIndex, const FMatrix& WorldMatrix, FGeometryCacheMeshData* MeshData)
+{
+	// Ensure sections array is long enough
+	if (TrackSections.Num() <= SectionIndex)
+	{
+		TrackSections.SetNum(SectionIndex + 1, false);
+	}
+
+	// Reset this section (in case it already existed)
+	FTrackRenderData& NewSection = TrackSections[SectionIndex];
+	NewSection.Reset();
+
+	// Number of Vertices
+	const int32 NumVerts = MeshData->Vertices.Num();
+
+	// Copy index buffer (clamping to vertex range)
+	const int32 NumTriIndices = NumVerts;
+	NewSection.IndexBuffer.Reset();
+	NewSection.IndexBuffer.AddUninitialized(NumTriIndices);
+	for (int32 IndexIdx = 0; IndexIdx < NumTriIndices; IndexIdx++)
+	{
+		NewSection.IndexBuffer[IndexIdx] = FMath::Min(IndexIdx, NumVerts - 1);
+	}
+	
+	// Store Matrix and MeshData-pointer
+	NewSection.WorldMatrix = WorldMatrix;
+	NewSection.MeshData = MeshData;
+
+	UpdateLocalBounds(); // Update overall bounds
+	MarkRenderStateDirty(); // Recreate scene proxy
+}
+
+void UGeometryCacheComponent::UpdateTrackSectionMeshData(int32 SectionIndex, FGeometryCacheMeshData* MeshData)
+{
+	// Check if the index is in range and Vertices contains any data
+	check(SectionIndex < TrackSections.Num() && MeshData != nullptr && "Invalid SectionIndex or Mesh Data");
+	// Reset this section (in case it already existed)
+	FTrackRenderData& UpdateSection = TrackSections[SectionIndex];
+	
+	// Number of Vertices
+	const int32 NumVerts = MeshData->Vertices.Num();
+	// Determine if we actually need a new index buffer
+	if (UpdateSection.IndexBuffer.Num() != MeshData->Vertices.Num())
+	{
+		// Copy index buffer (clamping to vertex range)
+		const int32 NumTriIndices = NumVerts;
+		UpdateSection.IndexBuffer.Reset();
+		UpdateSection.IndexBuffer.AddUninitialized(NumTriIndices);
+		for (int32 IndexIdx = 0; IndexIdx < NumTriIndices; IndexIdx++)
+		{
+			UpdateSection.IndexBuffer[IndexIdx] = FMath::Min(IndexIdx, NumVerts - 1);
+		}
+
+		if (SceneProxy)
+		{
+			UpdateTrackSectionIndexbuffer(SectionIndex, UpdateSection.IndexBuffer);
+		}
+	}
+
+	// Update MeshDataPointer
+	UpdateSection.MeshData = MeshData;
+
+	if (SceneProxy)
+	{		
+		UpdateTrackSectionVertexbuffer(SectionIndex, MeshData);
+		// Will update the bounds (used for frustum culling)
+		MarkRenderTransformDirty();
+	}
+	else
+	{
+		// Recreate scene proxy
+		MarkRenderStateDirty(); 
+	}
+
+	MarkRenderStateDirty();
+
+	// Update overall bounds	
+	UpdateLocalBounds(); 	
+}
+
+void UGeometryCacheComponent::UpdateTrackSectionMatrixData(int32 SectionIndex, const FMatrix& WorldMatrix)
+{
+	check(SectionIndex < TrackSections.Num() && "Invalid SectionIndex");
+	// Reset this section (in case it already existed)
+	FTrackRenderData& UpdateSection = TrackSections[SectionIndex];
+
+	// Update the WorldMatrix only
+	UpdateSection.WorldMatrix = WorldMatrix;
+
+	if (!IsRenderStateDirty() && SceneProxy != NULL)
+	{
+		// Update it within the scene proxy
+		SceneProxy->UpdateSectionWorldMatrix(SectionIndex, WorldMatrix);
+	}
+
+	// Update local bounds
+	UpdateLocalBounds();
+
+	// Mark transform Dirty
+	MarkRenderTransformDirty();
+}
+
+void UGeometryCacheComponent::UpdateTrackSectionVertexbuffer(int32 SectionIndex, FGeometryCacheMeshData* MeshData)
+{
+	ENQUEUE_UNIQUE_RENDER_COMMAND_THREEPARAMETER(
+		FUpdateVertexBufferCommand,
+		FGeometryCacheSceneProxy*, SceneProxy, SceneProxy,
+		const int32, Index, SectionIndex,
+		FGeometryCacheMeshData*, Data, MeshData,
+		{
+		SceneProxy->UpdateSectionVertexBuffer(Index, Data);
+	});
+}
+
+void UGeometryCacheComponent::UpdateTrackSectionIndexbuffer(int32 SectionIndex, const TArray<int32>& Indices)
+{
+	ENQUEUE_UNIQUE_RENDER_COMMAND_THREEPARAMETER(
+		FUpdateVertexBufferCommand,
+		FGeometryCacheSceneProxy*, SceneProxy, SceneProxy,
+		const int32, Index, SectionIndex,
+		const TArray<int32>&, Data, Indices,
+		{
+		SceneProxy->UpdateSectionIndexBuffer(Index, Data);
+	});
+}
+
+void UGeometryCacheComponent::OnObjectReimported(UGeometryCache* ImportedGeometryCache)
+{
+	if (GeometryCache == ImportedGeometryCache)
+	{
+		ReleaseResources();
+		GeometryCache = ImportedGeometryCache;
+
+		if (GeometryCache != NULL)
+		{
+			// Refresh NumTracks and clear Index Arrays
+			NumTracks = GeometryCache->Tracks.Num();
+			TrackMeshSampleIndices.Empty(GeometryCache->Tracks.Num());
+			TrackMatrixSampleIndices.Empty(GeometryCache->Tracks.Num());
+
+			Duration = 0.0f;
+			// Create mesh sections for each GeometryCacheTrack
+			for (int32 TrackIndex = 0; TrackIndex < NumTracks; ++TrackIndex)
+			{
+				UGeometryCacheTrack* Track = GeometryCache->Tracks[TrackIndex];
+				FMatrix WorldMatrix;
+				int32 MeshSampleIndex = -1;
+				int32 MatrixSampleIndex = -1;
+				FGeometryCacheMeshData* MeshData = NULL;
+
+				// Retrieve the matrix/mesh data and the appropriate sample indices
+				Track->UpdateMatrixData(ElapsedTime + StartTimeOffset, bLooping, MatrixSampleIndex, WorldMatrix);
+				Track->UpdateMeshData(ElapsedTime + StartTimeOffset, bLooping, MeshSampleIndex, MeshData);
+
+				// First time so create rather than update the mesh sections
+				CreateTrackSection(TrackIndex, WorldMatrix, MeshData);
+
+				// Store the sample indices for both the mesh and matrix data
+				TrackMeshSampleIndices.Add(MeshSampleIndex);
+				TrackMatrixSampleIndices.Add(MatrixSampleIndex);
+
+				const float TrackMaxSampleTime = Track->GetMaxSampleTime();
+				Duration = (Duration > TrackMaxSampleTime) ? Duration : TrackMaxSampleTime;
+			}
+		}
+
+		if (SceneProxy)
+		{
+			SceneProxy->ClearSections();
+		}
+
+
+		MarkRenderStateDirty();
+	}
+}
+
+void UGeometryCacheComponent::Play()
+{
+	bRunning = true;
+	PlayDirection = 1.0f;
+}
+
+void UGeometryCacheComponent::PlayFromStart()
+{
+	ElapsedTime = 0.0f;
+	bRunning = true;
+	PlayDirection = 1.0f;
+}
+
+void UGeometryCacheComponent::Pause()
+{
+	bRunning = !bRunning;
+}
+
+void UGeometryCacheComponent::Stop()
+{
+	bRunning = false;
+}
+
+bool UGeometryCacheComponent::IsPlaying() const
+{
+	return bRunning;
+}
+
+bool UGeometryCacheComponent::IsLooping() const
+{
+	return bLooping;
+}
+
+void UGeometryCacheComponent::SetLooping(const bool bNewLooping)
+{
+	bLooping = bNewLooping;
+}
+
+bool UGeometryCacheComponent::IsPlayingReversed() const
+{
+	return FMath::IsNearlyEqual( PlayDirection, -1.0f );
+}
+
+float UGeometryCacheComponent::GetPlaybackSpeed() const
+{
+	return PlaybackSpeed;
+}
+
+void UGeometryCacheComponent::SetPlaybackSpeed(const float NewPlaybackSpeed)
+{
+	// Currently only positive play back speeds are supported
+	PlaybackSpeed = fabs( NewPlaybackSpeed );
+}
+
+void UGeometryCacheComponent::PlayReversedFromEnd()
+{
+	ElapsedTime = Duration;
+	PlayDirection = -1.0f;
+	bRunning = true;
+}
+
+void UGeometryCacheComponent::PlayReversed()
+{
+	PlayDirection = -1.0f;
+	bRunning = true;
+}
+
+void UGeometryCacheComponent::InvalidateTrackSampleIndices()
+{
+	for (int32& Index : TrackMeshSampleIndices)
+	{
+		Index = -1;
+	}
+
+	for (int32& Index : TrackMatrixSampleIndices)
+	{
+		Index = -1;
+	}
+}
+
+void UGeometryCacheComponent::ReleaseResources()
+{
+	GeometryCache = NULL;
+	NumTracks = 0;
+	TrackMatrixSampleIndices.Empty();
+	TrackMatrixSampleIndices.Empty();
+	TrackSections.Empty();
+	DetachFence.BeginFence();
+}
+
+void UGeometryCacheComponent::PreEditUndo()
+{
+	InvalidateTrackSampleIndices();
+	MarkRenderStateDirty();
+}
+
+void UGeometryCacheComponent::PostEditUndo()
+{
+	InvalidateTrackSampleIndices();
+	MarkRenderStateDirty();
+}
diff --git a/Engine/Source/Runtime/GeometryCache/Private/GeometryCacheModule.cpp b/Engine/Source/Runtime/GeometryCache/Private/GeometryCacheModule.cpp
new file mode 100644
--- /dev/null
+++ b/Engine/Source/Runtime/GeometryCache/Private/GeometryCacheModule.cpp
@@ -0,0 +1,41 @@
+// Copyright 1998-2015 Epic Games, Inc. All Rights Reserved.
+
+#include "GeometryCacheModulePrivatePCH.h"
+#include "AssetToolsModule.h"
+#include "AssetTypeActions_GeometryCache.h"
+
+#include "GeometryCacheModule.h"
+
+IMPLEMENT_MODULE(FGeometryCacheModule, GeometryCache)
+
+void FGeometryCacheModule::StartupModule()
+{
+	// This code will execute after your module is loaded into memory (but after global variables are initialized, of course.)
+
+	FAssetToolsModule& AssetToolsModule = FAssetToolsModule::GetModule();
+
+	IAssetTools& AssetTools = AssetToolsModule.Get();
+
+	AssetTools.RegisterAssetTypeActions(MakeShareable(new FAssetTypeActions_GeometryCache));
+
+	// Register to be notified when an asset is reimported
+	OnAssetReimportDelegateHandle = FEditorDelegates::OnAssetReimport.AddRaw(this, &FGeometryCacheModule::OnObjectReimported);
+
+}
+
+void FGeometryCacheModule::ShutdownModule()
+{
+	// This function may be called during shutdown to clean up your module.  For modules that support dynamic reloading,
+	// we call this function before unloading the module.
+}
+
+void FGeometryCacheModule::OnObjectReimported(UObject* InObject)
+{
+	if (UGeometryCache* GeometryCache = Cast<UGeometryCache>(InObject))
+	{
+		for (TObjectIterator<UGeometryCacheComponent> CacheIt; CacheIt; ++CacheIt)
+		{
+			CacheIt->OnObjectReimported(GeometryCache);
+		}
+	}
+}
diff --git a/Engine/Source/Runtime/GeometryCache/Private/GeometryCacheModulePrivatePCH.h b/Engine/Source/Runtime/GeometryCache/Private/GeometryCacheModulePrivatePCH.h
new file mode 100644
--- /dev/null
+++ b/Engine/Source/Runtime/GeometryCache/Private/GeometryCacheModulePrivatePCH.h
@@ -0,0 +1,7 @@
+// Copyright 1998-2015 Epic Games, Inc. All Rights Reserved.
+
+#include "GeometryCacheModule.h"
+
+// You should place include statements to your module's private header files here.  You only need to
+// add includes for headers that are used in most of your module's source files though.
+#include "GeometryCache.h"
diff --git a/Engine/Source/Runtime/GeometryCache/Private/GeometryCacheSceneProxy.cpp b/Engine/Source/Runtime/GeometryCache/Private/GeometryCacheSceneProxy.cpp
new file mode 100644
--- /dev/null
+++ b/Engine/Source/Runtime/GeometryCache/Private/GeometryCacheSceneProxy.cpp
@@ -0,0 +1,340 @@
+// Copyright 1998-2015 Epic Games, Inc. All Rights Reserved.
+
+#include "GeometryCacheModulePrivatePCH.h"
+#include "GeometryCacheSceneProxy.h"
+
+FGeometryCacheSceneProxy::FGeometryCacheSceneProxy(UGeometryCacheComponent* Component) : FPrimitiveSceneProxy(Component)
+, MaterialRelevance(Component->GetMaterialRelevance(GetScene().GetFeatureLevel()))
+{
+	// Copy each section
+	const int32 NumSections = Component->TrackSections.Num();
+	Sections.AddZeroed(NumSections);
+	for (int SectionIdx = 0; SectionIdx < NumSections; SectionIdx++)
+	{
+		FTrackRenderData& SrcSection = Component->TrackSections[SectionIdx];
+		
+		if (SrcSection.IndexBuffer.Num() > 0)
+		{
+			FGeomCacheTrackProxy* NewSection = new FGeomCacheTrackProxy();
+
+			NewSection->WorldMatrix = SrcSection.WorldMatrix;
+			FGeometryCacheMeshData* MeshData = NewSection->MeshData = SrcSection.MeshData;
+
+			// Copy data from vertex buffer
+			const int32 NumVerts = MeshData->Vertices.Num();
+
+			// Allocate verts
+			NewSection->VertexBuffer.Vertices.Empty(NumVerts);
+			// Copy verts
+			NewSection->VertexBuffer.Vertices.Append(MeshData->Vertices);
+
+			// Copy index buffer
+			NewSection->IndexBuffer.Indices = SrcSection.IndexBuffer;
+
+			// Init vertex factory
+			NewSection->VertexFactory.Init(&NewSection->VertexBuffer);
+
+			// Enqueue initialization of render resource
+			BeginInitResource(&NewSection->VertexBuffer);
+			BeginInitResource(&NewSection->IndexBuffer);
+			BeginInitResource(&NewSection->VertexFactory);
+
+			// Grab materials
+			for (FGeometryCacheMeshBatchInfo& BatchInfo : MeshData->BatchesInfo)
+			{
+				UMaterialInterface* Material = Component->GetMaterial(BatchInfo.MaterialIndex);
+				if (Material == NULL)
+				{
+					Material = UMaterial::GetDefaultMaterial(MD_Surface);
+				}
+
+				NewSection->Materials.Push(Material);
+			}			
+
+			// Save ref to new section
+			Sections[SectionIdx] = NewSection;
+		}
+	}
+}
+
+FGeometryCacheSceneProxy::~FGeometryCacheSceneProxy()
+{
+	for (FGeomCacheTrackProxy* Section : Sections)
+	{
+		if (Section != nullptr)
+		{
+			Section->VertexBuffer.ReleaseResource();
+			Section->IndexBuffer.ReleaseResource();
+			Section->VertexFactory.ReleaseResource();
+			delete Section;
+		}
+	}
+	Sections.Empty();
+}
+
+void FGeometryCacheSceneProxy::GetDynamicMeshElements(const TArray<const FSceneView*>& Views, const FSceneViewFamily& ViewFamily, uint32 VisibilityMap, FMeshElementCollector& Collector) const
+{
+	SCOPE_CYCLE_COUNTER(STAT_GeometryCacheSceneProxy_GetMeshElements);
+
+	// Set up wireframe material (if needed)
+	const bool bWireframe = AllowDebugViewmodes() && ViewFamily.EngineShowFlags.Wireframe;
+
+	FColoredMaterialRenderProxy* WireframeMaterialInstance = NULL;
+	if (bWireframe)
+	{
+		WireframeMaterialInstance = new FColoredMaterialRenderProxy(
+			GEngine->WireframeMaterial ? GEngine->WireframeMaterial->GetRenderProxy(IsSelected()) : NULL,
+			FLinearColor(0, 0.5f, 1.f)
+			);
+
+		Collector.RegisterOneFrameMaterialProxy(WireframeMaterialInstance);
+	}
+	
+	// Iterate over sections	
+	for (const FGeomCacheTrackProxy* TrackProxy : Sections )
+	{
+		// QQQ
+		if (TrackProxy != nullptr)
+		{
+			INC_DWORD_STAT_BY(STAT_GeometryCacheSceneProxy_MeshBatchCount, TrackProxy->MeshData->BatchesInfo.Num());
+
+			int32 BatchIndex = 0;
+			for (FGeometryCacheMeshBatchInfo& BatchInfo : TrackProxy->MeshData->BatchesInfo)
+			{
+				FMaterialRenderProxy* MaterialProxy = bWireframe ? WireframeMaterialInstance : TrackProxy->Materials[BatchIndex]->GetRenderProxy(IsSelected());
+
+				for (int32 ViewIndex = 0; ViewIndex < Views.Num(); ViewIndex++)
+				{
+					if (VisibilityMap & (1 << ViewIndex))
+					{
+						const FSceneView* View = Views[ViewIndex];
+						// Draw the mesh.
+						FMeshBatch& Mesh = Collector.AllocateMesh();
+						FMeshBatchElement& BatchElement = Mesh.Elements[0];
+						BatchElement.IndexBuffer = &TrackProxy->IndexBuffer;
+						Mesh.bWireframe = bWireframe;
+						Mesh.VertexFactory = &TrackProxy->VertexFactory;
+						Mesh.MaterialRenderProxy = MaterialProxy;
+						BatchElement.PrimitiveUniformBuffer = CreatePrimitiveUniformBufferImmediate(TrackProxy->WorldMatrix * GetLocalToWorld(), GetBounds(), GetLocalBounds(), true, UseEditorDepthTest());
+						BatchElement.FirstIndex = BatchInfo.StartIndex;
+						BatchElement.NumPrimitives = BatchInfo.NumTriangles;
+						BatchElement.MinVertexIndex = 0;
+						BatchElement.MaxVertexIndex = TrackProxy->VertexBuffer.Vertices.Num() - 1;
+						Mesh.ReverseCulling = IsLocalToWorldDeterminantNegative();
+						Mesh.Type = PT_TriangleList;
+						Mesh.DepthPriorityGroup = SDPG_World;
+						Mesh.bCanApplyViewModeOverrides = false;
+						Collector.AddMesh(ViewIndex, Mesh);
+
+						INC_DWORD_STAT_BY(STAT_GeometryCacheSceneProxy_TriangleCount, BatchElement.NumPrimitives);
+					}
+				}
+
+				++BatchIndex;
+			}			
+		}
+		
+	}
+
+	// Draw bounds
+#if !(UE_BUILD_SHIPPING || UE_BUILD_TEST)
+	for (int32 ViewIndex = 0; ViewIndex < Views.Num(); ViewIndex++)
+	{
+		if (VisibilityMap & (1 << ViewIndex))
+		{
+			// Render bounds
+			RenderBounds(Collector.GetPDI(ViewIndex), ViewFamily.EngineShowFlags, GetBounds(), IsSelected());
+		}
+	}
+#endif
+}
+
+FPrimitiveViewRelevance FGeometryCacheSceneProxy::GetViewRelevance(const FSceneView* View)
+{
+	FPrimitiveViewRelevance Result;
+	Result.bDrawRelevance = IsShown(View);
+	Result.bShadowRelevance = IsShadowCast(View);
+	Result.bDynamicRelevance = true;
+	MaterialRelevance.SetPrimitiveViewRelevance(Result);
+	return Result;
+}
+
+bool FGeometryCacheSceneProxy::CanBeOccluded() const
+{
+	return !MaterialRelevance.bDisableDepthTest;
+}
+
+uint32 FGeometryCacheSceneProxy::GetMemoryFootprint(void) const
+{
+	return(sizeof(*this) + GetAllocatedSize());
+}
+
+uint32 FGeometryCacheSceneProxy::GetAllocatedSize(void) const
+{
+	return(FPrimitiveSceneProxy::GetAllocatedSize());
+}
+
+void FGeometryCacheSceneProxy::UpdateSectionWorldMatrix(const int32 SectionIndex, const FMatrix& WorldMatrix)
+{
+	check(SectionIndex < Sections.Num() && "Section Index out of range");
+	Sections[SectionIndex]->WorldMatrix = WorldMatrix;
+}
+
+
+void FGeometryCacheSceneProxy::UpdateSectionVertexBuffer(const int32 SectionIndex, FGeometryCacheMeshData* MeshData)
+{
+	check(SectionIndex < Sections.Num() && "Section Index out of range");
+	check(IsInRenderingThread());
+
+	Sections[SectionIndex]->MeshData = MeshData;
+
+	const bool bRecreate = Sections[SectionIndex]->VertexBuffer.Vertices.Num() != Sections[SectionIndex]->MeshData->Vertices.Num();
+	Sections[SectionIndex]->VertexBuffer.Vertices.Empty(Sections[SectionIndex]->MeshData->Vertices.Num());
+	Sections[SectionIndex]->VertexBuffer.Vertices.Append(Sections[SectionIndex]->MeshData->Vertices);
+
+	if (bRecreate)
+	{
+		Sections[SectionIndex]->VertexBuffer.InitRHI();
+	}
+	else
+	{
+		Sections[SectionIndex]->VertexBuffer.UpdateRHI();
+	}
+}
+
+void FGeometryCacheSceneProxy::UpdateSectionIndexBuffer(const int32 SectionIndex, const TArray<int32>& Indices)
+{
+	check(SectionIndex < Sections.Num() && "Section Index out of range");
+	check(IsInRenderingThread());
+
+	const bool bRecreate = Sections[SectionIndex]->IndexBuffer.Indices.Num() != Indices.Num();
+	
+	Sections[SectionIndex]->IndexBuffer.Indices.Empty(Indices.Num());
+	Sections[SectionIndex]->IndexBuffer.Indices.Append(Indices);
+	if (bRecreate)
+	{
+		Sections[SectionIndex]->IndexBuffer.InitRHI();
+	}
+	else
+	{
+		Sections[SectionIndex]->IndexBuffer.UpdateRHI();
+	}	
+}
+
+void FGeometryCacheSceneProxy::ClearSections()
+{
+	Sections.Empty();
+}
+
+FGeomCacheVertexFactory::FGeomCacheVertexFactory()
+{
+
+}
+
+void FGeomCacheVertexFactory::Init_RenderThread(const FGeomCacheVertexBuffer* VertexBuffer)
+{
+	check(IsInRenderingThread());
+
+	// Initialize the vertex factory's stream components.
+	DataType NewData;
+	NewData.PositionComponent = STRUCTMEMBER_VERTEXSTREAMCOMPONENT(VertexBuffer, FDynamicMeshVertex, Position, VET_Float3);
+	NewData.TextureCoordinates.Add(
+		FVertexStreamComponent(VertexBuffer, STRUCT_OFFSET(FDynamicMeshVertex, TextureCoordinate), sizeof(FDynamicMeshVertex), VET_Float2)
+		);
+	NewData.TangentBasisComponents[0] = STRUCTMEMBER_VERTEXSTREAMCOMPONENT(VertexBuffer, FDynamicMeshVertex, TangentX, VET_PackedNormal);
+	NewData.TangentBasisComponents[1] = STRUCTMEMBER_VERTEXSTREAMCOMPONENT(VertexBuffer, FDynamicMeshVertex, TangentZ, VET_PackedNormal);
+	NewData.ColorComponent = STRUCTMEMBER_VERTEXSTREAMCOMPONENT(VertexBuffer, FDynamicMeshVertex, Color, VET_Color);
+	SetData(NewData);
+}
+
+void FGeomCacheVertexFactory::Init(const FGeomCacheVertexBuffer* VertexBuffer)
+{
+	if (IsInRenderingThread())
+	{
+		Init_RenderThread(VertexBuffer);
+	}
+	else
+	{
+		ENQUEUE_UNIQUE_RENDER_COMMAND_TWOPARAMETER(
+			InitGeomCacheVertexFactory,
+			FGeomCacheVertexFactory*, VertexFactory, this,
+			const FGeomCacheVertexBuffer*, VertexBuffer, VertexBuffer,
+			{
+			VertexFactory->Init_RenderThread(VertexBuffer);
+		});
+	}
+}
+
+void FGeomCacheIndexBuffer::InitRHI()
+{
+	FRHIResourceCreateInfo CreateInfo;
+	IndexBufferRHI = RHICreateIndexBuffer(sizeof(int32), Indices.Num() * sizeof(int32), BUF_Static, CreateInfo);
+
+	// Write the indices to the index buffer.
+	void* Buffer = RHILockIndexBuffer(IndexBufferRHI, 0, Indices.Num() * sizeof(int32), RLM_WriteOnly);
+	FMemory::Memcpy(Buffer, Indices.GetData(), Indices.Num() * sizeof(int32));
+	RHIUnlockIndexBuffer(IndexBufferRHI);
+}
+
+void FGeomCacheIndexBuffer::UpdateRHI()
+{
+	// Copy the index data into the index buffer.
+	void* Buffer = RHILockIndexBuffer(IndexBufferRHI, 0, Indices.Num() * sizeof(int32), RLM_WriteOnly);
+	FMemory::Memcpy(Buffer, Indices.GetData(), Indices.Num() * sizeof(int32));
+	RHIUnlockIndexBuffer(IndexBufferRHI);
+}
+
+void FGeomCacheVertexBuffer::InitRHI()
+{
+	const uint32 SizeInBytes = Vertices.Num() * sizeof(FDynamicMeshVertex);
+
+	FGeomCacheVertexResourceArray ResourceArray(Vertices.GetData(), SizeInBytes);
+	FRHIResourceCreateInfo CreateInfo(&ResourceArray);
+	VertexBufferRHI = RHICreateVertexBuffer(SizeInBytes, BUF_Static, CreateInfo);
+}
+
+
+void FGeomCacheVertexBuffer::UpdateRHI()
+{
+	// Copy the vertex data into the vertex buffer.
+	void* VertexBufferData = RHILockVertexBuffer(VertexBufferRHI, 0, Vertices.Num() * sizeof(FDynamicMeshVertex), RLM_WriteOnly);
+	FMemory::Memcpy(VertexBufferData, Vertices.GetData(), Vertices.Num() * sizeof(FDynamicMeshVertex));
+	RHIUnlockVertexBuffer(VertexBufferRHI);
+}
+
+FGeomCacheVertexResourceArray::FGeomCacheVertexResourceArray(void* InData, uint32 InSize) : Data(InData)
+, Size(InSize)
+{
+
+}
+
+const void* FGeomCacheVertexResourceArray::GetResourceData() const
+{
+	return Data;
+}
+
+uint32 FGeomCacheVertexResourceArray::GetResourceDataSize() const
+{
+	return Size;
+}
+
+void FGeomCacheVertexResourceArray::Discard()
+{
+
+}
+
+bool FGeomCacheVertexResourceArray::IsStatic() const
+{
+	return false;
+}
+
+bool FGeomCacheVertexResourceArray::GetAllowCPUAccess() const
+{
+	return false;
+}
+
+void FGeomCacheVertexResourceArray::SetAllowCPUAccess(bool bInNeedsCPUAccess)
+{
+
+}
+
diff --git a/Engine/Source/Runtime/GeometryCache/Private/GeometryCacheSceneProxy.h b/Engine/Source/Runtime/GeometryCache/Private/GeometryCacheSceneProxy.h
new file mode 100644
--- /dev/null
+++ b/Engine/Source/Runtime/GeometryCache/Private/GeometryCacheSceneProxy.h
@@ -0,0 +1,118 @@
+// Copyright 1998-2015 Epic Games, Inc. All Rights Reserved.
+#pragma once
+
+#include "DynamicMeshBuilder.h"
+#include "Stats.h"
+
+DECLARE_STATS_GROUP(TEXT("GeometryCache"), STATGROUP_GeometryCache, STATCAT_Advanced);
+DECLARE_CYCLE_STAT(TEXT("MeshTime"), STAT_GeometryCacheSceneProxy_GetMeshElements, STATGROUP_GeometryCache );
+DECLARE_DWORD_COUNTER_STAT(TEXT("Triangle Count"), STAT_GeometryCacheSceneProxy_TriangleCount, STATGROUP_GeometryCache);
+DECLARE_DWORD_COUNTER_STAT(TEXT("Section Count"), STAT_GeometryCacheSceneProxy_MeshBatchCount, STATGROUP_GeometryCache);
+
+/** Resource array to pass  */
+class GEOMETRYCACHE_API FGeomCacheVertexResourceArray : public FResourceArrayInterface
+{
+public:
+	FGeomCacheVertexResourceArray(void* InData, uint32 InSize);
+
+	virtual const void* GetResourceData() const override;
+	virtual uint32 GetResourceDataSize() const override;
+	virtual void Discard() override;
+	virtual bool IsStatic() const override;
+	virtual bool GetAllowCPUAccess() const override;
+	virtual void SetAllowCPUAccess(bool bInNeedsCPUAccess) override;
+
+private:
+	void* Data;
+	uint32 Size;
+};
+
+/** Vertex Buffer */
+class GEOMETRYCACHE_API FGeomCacheVertexBuffer : public FVertexBuffer
+{
+public:
+	TArray<FDynamicMeshVertex> Vertices;
+
+	virtual void InitRHI() override;
+
+	void UpdateRHI();
+};
+
+/** Index Buffer */
+class GEOMETRYCACHE_API FGeomCacheIndexBuffer : public FIndexBuffer
+{
+public:
+	TArray<int32> Indices;
+
+	virtual void InitRHI() override;
+
+	void UpdateRHI();
+};
+
+/** Vertex Factory */
+class GEOMETRYCACHE_API FGeomCacheVertexFactory : public FLocalVertexFactory
+{
+public:
+
+	FGeomCacheVertexFactory();
+
+	/** Init function that should only be called on render thread. */
+	void Init_RenderThread(const FGeomCacheVertexBuffer* VertexBuffer);
+
+	/** Init function that can be called on any thread, and will do the right thing (enqueue command if called on main thread) */
+	void Init(const FGeomCacheVertexBuffer* VertexBuffer);
+};
+
+class GEOMETRYCACHE_API FGeomCacheTrackProxy
+{
+public:
+	/** MeshData storing information used for rendering this Track */
+	FGeometryCacheMeshData* MeshData;
+
+	/** Material applied to this Track */
+	TArray<UMaterialInterface*> Materials;
+	/** Vertex buffer for this Track */
+	FGeomCacheVertexBuffer VertexBuffer;
+	/** Index buffer for this Track */
+	FGeomCacheIndexBuffer IndexBuffer;
+	/** Vertex factory for this Track */
+	FGeomCacheVertexFactory VertexFactory;
+	/** World Matrix for this Track */
+	FMatrix WorldMatrix;
+};
+
+/** Procedural mesh scene proxy */
+class GEOMETRYCACHE_API FGeometryCacheSceneProxy : public FPrimitiveSceneProxy
+{
+public:
+	// QQ interfaces/comments
+	FGeometryCacheSceneProxy(UGeometryCacheComponent* Component);
+
+	virtual ~FGeometryCacheSceneProxy();
+
+	virtual void GetDynamicMeshElements(const TArray<const FSceneView*>& Views, const FSceneViewFamily& ViewFamily, uint32 VisibilityMap, FMeshElementCollector& Collector) const override;
+
+	virtual FPrimitiveViewRelevance GetViewRelevance(const FSceneView* View);
+
+	virtual bool CanBeOccluded() const override;
+
+	virtual uint32 GetMemoryFootprint(void) const;
+
+	uint32 GetAllocatedSize(void) const;
+
+	/** Update world matrix for specific section */
+	void UpdateSectionWorldMatrix(const int32 SectionIndex, const FMatrix& WorldMatrix);
+	/** Update vertex buffer for specific section */
+	void UpdateSectionVertexBuffer(const int32 SectionIndex, FGeometryCacheMeshData* MeshData );
+	/** Update index buffer for specific section */
+	void UpdateSectionIndexBuffer(const int32 SectionIndex, const TArray<int32>& Indices);
+
+	/** Clears the Sections array*/
+	void ClearSections();
+	
+private:	
+	FMaterialRelevance MaterialRelevance;
+
+	/** Array of Track Proxies */
+	TArray<FGeomCacheTrackProxy*> Sections;
+};
diff --git a/Engine/Source/Runtime/GeometryCache/Private/GeometryCacheTrack.cpp b/Engine/Source/Runtime/GeometryCache/Private/GeometryCacheTrack.cpp
new file mode 100644
--- /dev/null
+++ b/Engine/Source/Runtime/GeometryCache/Private/GeometryCacheTrack.cpp
@@ -0,0 +1,116 @@
+// Copyright 1998-2015 Epic Games, Inc. All Rights Reserved.
+
+#include "GeometryCacheModulePrivatePCH.h"
+#include "GeometryCacheTrack.h"
+
+GEOMETRYCACHE_API UGeometryCacheTrack::UGeometryCacheTrack(const FObjectInitializer& ObjectInitializer /*= FObjectInitializer::Get()*/) : UObject(ObjectInitializer)
+{
+	NumMaterials = 0;
+}
+
+const bool UGeometryCacheTrack::UpdateMeshData(const float Time, const bool bLooping, int32& InOutMeshSampleIndex, FGeometryCacheMeshData*& OutMeshData)
+{
+	return false;
+}
+
+const bool UGeometryCacheTrack::UpdateMatrixData(const float Time, const bool bLooping, int32& InOutMatrixSampleIndex, FMatrix& OutWorldMatrix)
+{
+	// Retrieve sample index from Time
+	const uint32 MatrixSampleIndex = FindSampleIndexFromTime(MatrixSampleTimes, Time, bLooping);
+
+	// Update the Matrix and Index if MatrixSampleIndex is different from the stored InOutMatrixSampleIndex
+	if (MatrixSampleIndex != InOutMatrixSampleIndex)
+	{
+		InOutMatrixSampleIndex = MatrixSampleIndex;
+		OutWorldMatrix = MatrixSamples[MatrixSampleIndex];
+		return true;
+	}
+
+	return false;
+}
+
+UGeometryCacheTrack::~UGeometryCacheTrack()
+{
+	MatrixSamples.Empty();
+	MatrixSampleTimes.Empty();
+}
+
+void UGeometryCacheTrack::Serialize(FArchive& Ar)
+{
+	Ar << MatrixSamples;
+	Ar << MatrixSampleTimes;
+	Ar << NumMaterials;
+}
+
+void UGeometryCacheTrack::SetMatrixSamples(const TArray<FMatrix>& Matrices, const TArray<float>& SampleTimes)
+{
+	// Copy Matrix samples and sample-times
+	MatrixSamples.Append(Matrices);
+	MatrixSampleTimes.Append(SampleTimes);
+}
+
+const float UGeometryCacheTrack::GetMaxSampleTime() const
+{
+	// If there are sample-times available return the (maximal) time from the last sample
+	if ( MatrixSampleTimes.Num() > 0 )
+	{
+		return MatrixSampleTimes.Last();
+	}
+
+	// Otherwise no data/times availabel
+	return 0.0f;	
+}
+
+const uint32 UGeometryCacheTrack::FindSampleIndexFromTime(const TArray<float>& SampleTimes, const float Time, const bool bLooping)
+{
+	// No index possible
+	if (SampleTimes.Num() == 0)
+	{
+		return 0;
+	}
+
+	// Modulo the incoming Time if the animation is played on a loop
+	float SampleTime = Time;
+	if (bLooping)
+	{
+		SampleTime = fmod(Time, SampleTimes[SampleTimes.Num() - 1]);
+	}
+
+	// Binary searching for closest (floored) SampleIndex 
+	uint32 MinIndex = 0;
+	uint32 MaxIndex = SampleTimes.Num() - 1;
+	if (SampleTime > SampleTimes[MaxIndex])
+	{
+		return MaxIndex;
+	}
+	else if (SampleTime < SampleTimes[MinIndex])
+	{
+		return MinIndex;
+	}
+
+	while (MaxIndex > 0 && MaxIndex >= MinIndex)
+	{
+		uint32 Mid = (MinIndex + MaxIndex) / 2;
+		if (SampleTime > SampleTimes[Mid])
+		{
+			MinIndex = Mid + 1;
+		}
+		else
+		{
+			MaxIndex = Mid - 1;
+		}
+	}
+
+	return MinIndex;
+}
+
+SIZE_T UGeometryCacheTrack::GetResourceSize(EResourceSizeMode::Type Mode)
+{
+	SIZE_T ResourceSize = 0;
+
+	// Determine resource size from data that is serialized
+	ResourceSize += MatrixSamples.Num() * sizeof(FMatrix);
+	ResourceSize += MatrixSampleTimes.Num() * sizeof(float);
+
+	return ResourceSize;
+}
diff --git a/Engine/Source/Runtime/GeometryCache/Private/GeometryCacheTrackFlipbookAnimation.cpp b/Engine/Source/Runtime/GeometryCache/Private/GeometryCacheTrackFlipbookAnimation.cpp
new file mode 100644
--- /dev/null
+++ b/Engine/Source/Runtime/GeometryCache/Private/GeometryCacheTrackFlipbookAnimation.cpp
@@ -0,0 +1,95 @@
+// Copyright 1998-2015 Epic Games, Inc. All Rights Reserved.
+
+#include "GeometryCacheModulePrivatePCH.h"
+#include "GeometryCacheMeshData.h"
+
+#include "GeometryCacheTrackFlipbookAnimation.h"
+
+GEOMETRYCACHE_API UGeometryCacheTrack_FlipbookAnimation::UGeometryCacheTrack_FlipbookAnimation(const FObjectInitializer& ObjectInitializer /*= FObjectInitializer::Get()*/) : UGeometryCacheTrack(ObjectInitializer)
+{
+	NumMeshSamples = 0;
+}
+
+UGeometryCacheTrack_FlipbookAnimation::~UGeometryCacheTrack_FlipbookAnimation()
+{
+	NumMeshSamples = 0;
+	MeshSamples.Empty();
+	MeshSampleTimes.Empty();
+}
+
+SIZE_T UGeometryCacheTrack_FlipbookAnimation::GetResourceSize(EResourceSizeMode::Type Mode)
+{
+	// Determine resource size according to what is actually serialized
+	SIZE_T ResourceSize = 0;
+	ResourceSize += UGeometryCacheTrack::GetResourceSize(Mode);
+	for (int32 SampleIndex = 0; SampleIndex < MeshSamples.Num(); ++SampleIndex )
+	{
+		ResourceSize += MeshSamples[SampleIndex].GetResourceSize();
+	}
+	ResourceSize += sizeof(MeshSamples);
+	ResourceSize += MeshSampleTimes.Num() * sizeof(float);
+	ResourceSize += sizeof(MeshSampleTimes);
+	ResourceSize += sizeof(NumMeshSamples);
+
+	return ResourceSize;
+}
+
+void UGeometryCacheTrack_FlipbookAnimation::Serialize(FArchive& Ar)
+{
+	UGeometryCacheTrack::Serialize(Ar);
+	
+	Ar << NumMeshSamples;
+	Ar << MeshSamples;
+	Ar << MeshSampleTimes;
+}
+
+const bool UGeometryCacheTrack_FlipbookAnimation::UpdateMeshData(const float Time, const bool bLooping, int32& InOutMeshSampleIndex, FGeometryCacheMeshData*& OutMeshData)
+{
+	// Retrieve sample index from Time
+	const int32 MeshSampleIndex = FindSampleIndexFromTime(MeshSampleTimes, Time, bLooping);
+
+	// Update the Vertices and Index if MeshSampleIndex is different from the stored InOutMeshSampleIndex
+	if (MeshSampleIndex != InOutMeshSampleIndex)
+	{
+		OutMeshData = &MeshSamples[MeshSampleIndex];
+		InOutMeshSampleIndex = MeshSampleIndex;
+
+		return true;
+	}
+
+	return false;
+}
+
+const float UGeometryCacheTrack_FlipbookAnimation::GetMaxSampleTime() const
+{
+	const float BaseTime = UGeometryCacheTrack::GetMaxSampleTime();
+
+	if (MeshSampleTimes.Num() > 0)
+	{
+		const float MeshSampleTime = MeshSampleTimes.Last();
+		return (BaseTime > MeshSampleTime) ? BaseTime : MeshSampleTime;
+	}
+
+	return BaseTime;
+}
+
+void UGeometryCacheTrack_FlipbookAnimation::AddMeshSample(const FGeometryCacheMeshData& MeshData, const float SampleTime)
+{
+	MeshSamples.Add(MeshData);
+	MeshSampleTimes.Add(SampleTime);
+	NumMeshSamples++;
+
+	// Store the total number of materials within this track
+	if (MeshData.BatchesInfo.Num() > (int32)NumMaterials)
+	{
+		NumMaterials = MeshData.BatchesInfo.Num();
+	}
+}
+
+void UGeometryCacheTrack_FlipbookAnimation::BeginDestroy()
+{
+	Super::BeginDestroy();
+	NumMeshSamples = 0;
+	MeshSamples.Empty();
+	MeshSampleTimes.Empty();
+}
diff --git a/Engine/Source/Runtime/GeometryCache/Private/GeometryCacheTrackTransformAnimation.cpp b/Engine/Source/Runtime/GeometryCache/Private/GeometryCacheTrackTransformAnimation.cpp
new file mode 100644
--- /dev/null
+++ b/Engine/Source/Runtime/GeometryCache/Private/GeometryCacheTrackTransformAnimation.cpp
@@ -0,0 +1,43 @@
+// Copyright 1998-2015 Epic Games, Inc. All Rights Reserved.
+
+#include "GeometryCacheModulePrivatePCH.h"
+#include "GeometryCacheTrackTransformAnimation.h"
+
+
+GEOMETRYCACHE_API UGeometryCacheTrack_TransformAnimation::UGeometryCacheTrack_TransformAnimation(const FObjectInitializer& ObjectInitializer /*= FObjectInitializer::Get()*/) : UGeometryCacheTrack(ObjectInitializer)
+{
+}
+
+SIZE_T UGeometryCacheTrack_TransformAnimation::GetResourceSize(EResourceSizeMode::Type Mode)
+{
+	SIZE_T ResourceSize = 0;
+	ResourceSize += UGeometryCacheTrack::GetResourceSize(Mode);
+	ResourceSize += MeshData.GetResourceSize();
+
+	return ResourceSize;
+}
+
+void UGeometryCacheTrack_TransformAnimation::Serialize(FArchive& Ar)
+{
+	UGeometryCacheTrack::Serialize(Ar);
+	Ar << MeshData;
+}
+
+const bool UGeometryCacheTrack_TransformAnimation::UpdateMeshData(const float Time, const bool bLooping, int32& InOutMeshSampleIndex, FGeometryCacheMeshData*& OutMeshData)
+{
+	// If InOutMeshSampleIndex equals -1 (first creation) update the OutVertices and InOutMeshSampleIndex
+	if (InOutMeshSampleIndex == -1)
+	{
+		OutMeshData = &MeshData;
+		InOutMeshSampleIndex = 0;
+		return true;
+	}
+
+	return false;
+}
+
+void UGeometryCacheTrack_TransformAnimation::SetMesh(const FGeometryCacheMeshData& NewMeshData)
+{
+	MeshData = NewMeshData;
+	NumMaterials = NewMeshData.BatchesInfo.Num();
+}
diff --git a/Engine/Source/Runtime/GeometryCache/Public/GeometryCacheModule.h b/Engine/Source/Runtime/GeometryCache/Public/GeometryCacheModule.h
new file mode 100644
--- /dev/null
+++ b/Engine/Source/Runtime/GeometryCache/Public/GeometryCacheModule.h
@@ -0,0 +1,48 @@
+// Copyright 1998-2015 Epic Games, Inc. All Rights Reserved.
+
+#pragma once
+
+#include "ModuleManager.h"
+
+/**
+ * The public interface to this module
+ */
+class FGeometryCacheModule : public IModuleInterface
+{
+public:
+	virtual void StartupModule();
+	virtual void ShutdownModule();
+
+	/**
+	 * Singleton-like access to this module's interface.  This is just for convenience!
+	 * Beware of calling this during the shutdown phase, though.  Your module might have been unloaded already.
+	 *
+	 * @return Returns singleton instance, loading the module on demand if needed
+	 */
+	static inline FGeometryCacheModule& Get()
+	{
+		return FModuleManager::LoadModuleChecked< FGeometryCacheModule >("GeometryCache");
+	}
+
+	/**
+	 * Checks to see if this module is loaded and ready.  It is only valid to call Get() if IsAvailable() returns true.
+	 *
+	 * @return True if the module is loaded and ready to use
+	 */
+	static inline bool IsAvailable()
+	{
+		return FModuleManager::Get().IsModuleLoaded( "GeometryCache" );
+	}
+
+	/**
+	* OnObjectReimported, used to fix up rendering data for GeometryCacheComponent that's using the reimported GeometryCache
+	*
+	* @param InObject - Reimported object
+	* @return void
+	*/
+	void OnObjectReimported(UObject* InObject);
+	
+private:
+	FDelegateHandle OnAssetReimportDelegateHandle;
+};
+
