Attempt to fix a non-unity build.

diff --git a/Engine/Source/Runtime/GeometryCache/Private/GeometryCacheModule.cpp b/Engine/Source/Runtime/GeometryCache/Private/GeometryCacheModule.cpp
--- a/Engine/Source/Runtime/GeometryCache/Private/GeometryCacheModule.cpp
+++ b/Engine/Source/Runtime/GeometryCache/Private/GeometryCacheModule.cpp
@@ -5,6 +5,7 @@
 #include "AssetTypeActions_GeometryCache.h"
 
 #include "GeometryCacheModule.h"
+#include "GeometryCacheComponent.h"
 
 IMPLEMENT_MODULE(FGeometryCacheModule, GeometryCache)
 
diff --git a/Engine/Source/Runtime/GeometryCache/Private/GeometryCacheSceneProxy.cpp b/Engine/Source/Runtime/GeometryCache/Private/GeometryCacheSceneProxy.cpp
--- a/Engine/Source/Runtime/GeometryCache/Private/GeometryCacheSceneProxy.cpp
+++ b/Engine/Source/Runtime/GeometryCache/Private/GeometryCacheSceneProxy.cpp
@@ -2,6 +2,8 @@
 
 #include "GeometryCacheModulePrivatePCH.h"
 #include "GeometryCacheSceneProxy.h"
+#include "GeometryCacheComponent.h"
+#include "GeometryCacheMeshData.h"
 
 FGeometryCacheSceneProxy::FGeometryCacheSceneProxy(UGeometryCacheComponent* Component) : FPrimitiveSceneProxy(Component)
 , MaterialRelevance(Component->GetMaterialRelevance(GetScene().GetFeatureLevel()))
diff --git a/Engine/Source/Runtime/GeometryCache/Private/GeometryCacheSceneProxy.h b/Engine/Source/Runtime/GeometryCache/Private/GeometryCacheSceneProxy.h
--- a/Engine/Source/Runtime/GeometryCache/Private/GeometryCacheSceneProxy.h
+++ b/Engine/Source/Runtime/GeometryCache/Private/GeometryCacheSceneProxy.h
@@ -86,7 +86,7 @@ class GEOMETRYCACHE_API FGeometryCacheSceneProxy : public FPrimitiveSceneProxy
 {
 public:
 	// QQ interfaces/comments
-	FGeometryCacheSceneProxy(UGeometryCacheComponent* Component);
+	FGeometryCacheSceneProxy(class UGeometryCacheComponent* Component);
 
 	virtual ~FGeometryCacheSceneProxy();
 
