GeomCache/ABCImporter
ADDED thumbnail renderer for GeometryCache
ADDED BP creation from GeometryCache asset
ADDED de-registering of handlers during GeomCacheModule shutdown
FIXED bug with importing Materials, though recreating a clean scene in Maya also fixed it (bad Maya history)

diff --git a/Engine/Config/BaseEditor.ini b/Engine/Config/BaseEditor.ini
--- a/Engine/Config/BaseEditor.ini
+++ b/Engine/Config/BaseEditor.ini
@@ -40,6 +40,7 @@ MaxNumberOfLogLines=2000
 +RenderableThumbnailTypes=(ClassNeedingThumbnailName="/Script/Engine.SoundWave",RendererClassName="/Script/UnrealEd.SoundWaveThumbnailRenderer")
 +RenderableThumbnailTypes=(ClassNeedingThumbnailName="/Script/Engine.SlateBrushAsset",RendererClassName="/Script/UnrealEd.SlateBrushThumbnailRenderer")
 +RenderableThumbnailTypes=(ClassNeedingThumbnailName="/Script/Engine.SubsurfaceProfile",RendererClassName="/Script/UnrealEd.SubsurfaceProfileRenderer")
++RenderableThumbnailTypes=(ClassNeedingThumbnailName="/Script/GeometryCache.GeometryCache",RendererClassName="/Script/GeometryCacheEd.GeometryCacheThumbnailRenderer")
 
 [/Script/UnrealEd.DefaultSizedThumbnailRenderer]
 DefaultSizeX=512
diff --git a/Engine/Source/Developer/MeshUtilities/Private/MeshUtilities.cpp b/Engine/Source/Developer/MeshUtilities/Private/MeshUtilities.cpp
--- a/Engine/Source/Developer/MeshUtilities/Private/MeshUtilities.cpp
+++ b/Engine/Source/Developer/MeshUtilities/Private/MeshUtilities.cpp
@@ -124,7 +124,7 @@ private:
 	 * @param Chunks				Skinned mesh chunks from which to build the renderable model.
 	 * @param PointToOriginalMap	Maps a vertex's RawPointIdx to its index at import time.
 	 */
-	void BuildSkeletalModelFromChunks(FStaticLODModel& LODModel,const FReferenceSkeleton& RefSkeleton,TArray<FSkinnedMeshChunk*>& Chunks,const TArray<int32>& PointToOriginalMap);
+	void BuildSkeletalModelFromChunks(FStaticLODModel& LODModel, const FReferenceSkeleton& RefSkeleton, TArray<FSkinnedMeshChunk*>& Chunks, const TArray<int32>& PointToOriginalMap);
 
 	// IModuleInterface interface.
 	virtual void StartupModule() override;
@@ -170,7 +170,7 @@ private:
 		TArray<int32>& OutGlobalMaterialIndices
 		) const;
 
-	virtual void CalculateRawMeshTangents(FRawMesh& OutRawMesh, const FMeshBuildSettings& BuildSettings);
+	virtual void ExtractMeshDataForGeometryCache(FRawMesh& RawMesh, const FMeshBuildSettings& BuildSettings, TArray<FStaticMeshBuildVertex>& OutVertices, TArray<TArray<uint32> >& OutPerSectionIndices );
 
 	// Need to call some members from this class, (which is internal to this module)
 	friend class FStaticMeshUtilityBuilder;
@@ -5406,35 +5406,35 @@ bool FMeshUtilities::ConstructRawMesh(
 	return true;
 }
 
-void FMeshUtilities::CalculateRawMeshTangents(FRawMesh& OutRawMesh, const FMeshBuildSettings& BuildSettings)
+void FMeshUtilities::ExtractMeshDataForGeometryCache(FRawMesh& RawMesh, const FMeshBuildSettings& BuildSettings, TArray<FStaticMeshBuildVertex>& OutVertices, TArray<TArray<uint32> >& OutPerSectionIndices)
 {
-	int32 NumWedges = OutRawMesh.WedgeIndices.Num();
+	int32 NumWedges = RawMesh.WedgeIndices.Num();
 
 	// Figure out if we should recompute normals and tangents. By default generated LODs should not recompute normals
-	bool bRecomputeNormals = (BuildSettings.bRecomputeNormals) || OutRawMesh.WedgeTangentZ.Num() == 0;
-	bool bRecomputeTangents = (BuildSettings.bRecomputeTangents) || OutRawMesh.WedgeTangentX.Num() == 0 || OutRawMesh.WedgeTangentY.Num() == 0;
+	bool bRecomputeNormals = (BuildSettings.bRecomputeNormals) || RawMesh.WedgeTangentZ.Num() == 0;
+	bool bRecomputeTangents = (BuildSettings.bRecomputeTangents) || RawMesh.WedgeTangentX.Num() == 0 || RawMesh.WedgeTangentY.Num() == 0;
 
 	// Dump normals and tangents if we are recomputing them.
 	if (bRecomputeTangents)
 	{
-		OutRawMesh.WedgeTangentX.Empty(NumWedges);
-		OutRawMesh.WedgeTangentX.AddZeroed(NumWedges);
-		OutRawMesh.WedgeTangentY.Empty(NumWedges);
-		OutRawMesh.WedgeTangentY.AddZeroed(NumWedges);
+		RawMesh.WedgeTangentX.Empty(NumWedges);
+		RawMesh.WedgeTangentX.AddZeroed(NumWedges);
+		RawMesh.WedgeTangentY.Empty(NumWedges);
+		RawMesh.WedgeTangentY.AddZeroed(NumWedges);
 	}
 
 	if (bRecomputeNormals)
 	{
-		OutRawMesh.WedgeTangentZ.Empty(NumWedges);
-		OutRawMesh.WedgeTangentZ.AddZeroed(NumWedges);
+		RawMesh.WedgeTangentZ.Empty(NumWedges);
+		RawMesh.WedgeTangentZ.AddZeroed(NumWedges);
 	}
 
 	// Compute any missing tangents.
+	TMultiMap<int32, int32> OverlappingCorners;
 	if (bRecomputeNormals || bRecomputeTangents)
 	{
-		float ComparisonThreshold = GetComparisonThreshold(BuildSettings);
-		TMultiMap<int32, int32> OverlappingCorners;
-		FindOverlappingCorners(OverlappingCorners, OutRawMesh, ComparisonThreshold);
+		float ComparisonThreshold = GetComparisonThreshold(BuildSettings);		
+		FindOverlappingCorners(OverlappingCorners, RawMesh, ComparisonThreshold);
 
 		// Static meshes always blend normals of overlapping corners.
 		uint32 TangentOptions = ETangentOptions::BlendOverlappingNormals;
@@ -5445,18 +5445,39 @@ void FMeshUtilities::CalculateRawMeshTangents(FRawMesh& OutRawMesh, const FMeshB
 		}
 		if (BuildSettings.bUseMikkTSpace)
 		{
-			ComputeTangents_MikkTSpace(OutRawMesh, OverlappingCorners, TangentOptions);
+			ComputeTangents_MikkTSpace(RawMesh, OverlappingCorners, TangentOptions);
 		}
 		else
 		{
-			ComputeTangents(OutRawMesh, OverlappingCorners, TangentOptions);
+			ComputeTangents(RawMesh, OverlappingCorners, TangentOptions);
 		}
 	}
 
 	// At this point the mesh will have valid tangents.
-	check(OutRawMesh.WedgeTangentX.Num() == NumWedges);
-	check(OutRawMesh.WedgeTangentY.Num() == NumWedges);
-	check(OutRawMesh.WedgeTangentZ.Num() == NumWedges);
+	check(RawMesh.WedgeTangentX.Num() == NumWedges);
+	check(RawMesh.WedgeTangentY.Num() == NumWedges);
+	check(RawMesh.WedgeTangentZ.Num() == NumWedges);
+		
+	TArray<int32> OutWedgeMap;
+
+	int32 MaxMaterialIndex = 1;
+	for (int32 FaceIndex = 0; FaceIndex < RawMesh.FaceMaterialIndices.Num(); FaceIndex++)
+	{
+		MaxMaterialIndex = FMath::Max<int32>(RawMesh.FaceMaterialIndices[FaceIndex], MaxMaterialIndex);
+	}
+
+	for (int32 i = 0; i <= MaxMaterialIndex; ++i)
+	{
+		OutPerSectionIndices.Push( TArray<uint32>() );
+	}
+
+	BuildStaticMeshVertexAndIndexBuffers(OutVertices, OutPerSectionIndices, OutWedgeMap, RawMesh, OverlappingCorners, KINDA_SMALL_NUMBER, BuildSettings.BuildScale3D);
+
+	if (RawMesh.WedgeIndices.Num() < 100000 * 3)
+	{
+		CacheOptimizeVertexAndIndexBuffer(OutVertices, OutPerSectionIndices, OutWedgeMap);
+		check(OutWedgeMap.Num() == RawMesh.WedgeIndices.Num());
+	}
 }
 
 /*------------------------------------------------------------------------------
@@ -5925,6 +5946,7 @@ void FMeshUtilities::MergeStaticMeshComponents(const TArray<UStaticMeshComponent
 		{
 			UStaticMeshComponent* MeshComponent = ComponentsToMerge[MeshId];
 						
+
 			// We duplicate lower LOD in case this mesh has no LOD we want
 			int32 ExportLODIndex = FMath::Min(LODIndex, MeshComponent->StaticMesh->SourceModels.Num() - 1);
 
@@ -6301,7 +6323,6 @@ void FMeshUtilities::MergeStaticMeshComponents(const TArray<UStaticMeshComponent
 		OutMergedActorLocation = MergedMesh.Pivot;
 	}
 }
-	
 
 /*------------------------------------------------------------------------------
 	Mesh reduction.
diff --git a/Engine/Source/Developer/MeshUtilities/Public/MeshUtilities.h b/Engine/Source/Developer/MeshUtilities/Public/MeshUtilities.h
--- a/Engine/Source/Developer/MeshUtilities/Public/MeshUtilities.h
+++ b/Engine/Source/Developer/MeshUtilities/Public/MeshUtilities.h
@@ -255,11 +255,14 @@ public:
 		) = 0;
 
 	/**
-	* CalculateRawMeshTangents
+	* ExtractMeshDataForGeometryCache
 	*
-	* @param OutRawMesh - InOut raw Mesh for calculating tangents
+	* @param RawMesh - raw Mesh for calculating tangents
+	* @param BuildSettings - Buildsettings for calculating tangents
+	* @param OutVertices - Vertex buffer data for geometry cache
+	* @param OutPerSectionIndices - Index buffer data for geometry cache
 	* @return void
 	*/
-	virtual void CalculateRawMeshTangents(FRawMesh& OutRawMesh, const FMeshBuildSettings& BuildSettings) = 0;
+	virtual void ExtractMeshDataForGeometryCache(FRawMesh& RawMesh, const FMeshBuildSettings& BuildSettings, TArray<FStaticMeshBuildVertex>& OutVertices, TArray<TArray<uint32> >& OutPerSectionIndices) = 0;
 
 };
diff --git a/Engine/Source/Editor/GeometryCacheEd/Classes/ActorFactoryGeometryCache.h b/Engine/Source/Editor/GeometryCacheEd/Classes/ActorFactoryGeometryCache.h
new file mode 100644
--- /dev/null
+++ b/Engine/Source/Editor/GeometryCacheEd/Classes/ActorFactoryGeometryCache.h
@@ -0,0 +1,20 @@
+// Copyright 1998-2015 Epic Games, Inc. All Rights Reserved.
+
+#pragma once
+
+#include "UnrealEd.h"
+#include "ActorFactories/ActorFactory.h"
+#include "ActorFactoryGeometryCache.generated.h"
+
+/** Factory class for spawning and creating GeometryCacheActors */
+UCLASS(MinimalAPI, config = Editor)
+class UActorFactoryGeometryCache : public UActorFactory
+{
+	GENERATED_UCLASS_BODY()
+
+	// Begin UActorFactory Interface
+	virtual bool CanCreateActorFrom(const FAssetData& AssetData, FText& OutErrorMsg) override;
+	virtual void PostSpawnActor(UObject* Asset, AActor* NewActor) override;
+	virtual void PostCreateBlueprint(UObject* Asset, AActor* CDO) override;	
+	// End UActorFactory Interface
+};
diff --git a/Engine/Source/Editor/GeometryCacheEd/Classes/AssetTypeActions_GeometryCache.h b/Engine/Source/Editor/GeometryCacheEd/Classes/AssetTypeActions_GeometryCache.h
new file mode 100644
--- /dev/null
+++ b/Engine/Source/Editor/GeometryCacheEd/Classes/AssetTypeActions_GeometryCache.h
@@ -0,0 +1,23 @@
+// Copyright 1998-2015 Epic Games, Inc. All Rights Reserved.
+#pragma once
+
+#include "GeometryCache.h"
+#include "AssetTypeActions_Base.h"
+
+/** Type-actions class for GeometryCache assets*/
+class FAssetTypeActions_GeometryCache : public FAssetTypeActions_Base
+{
+public:
+	// IAssetTypeActions Implementation
+	virtual FText GetName() const override;
+	virtual FColor GetTypeColor() const override;
+	virtual UClass* GetSupportedClass() const override;
+	virtual bool HasActions(const TArray<UObject*>& InObjects) const override;
+	virtual void GetActions(const TArray<UObject*>& InObjects, FMenuBuilder& MenuBuilder) override;
+	virtual void OpenAssetEditor(const TArray<UObject*>& InObjects, TSharedPtr<class IToolkitHost> EditWithinLevelEditor = TSharedPtr<IToolkitHost>()) override;
+	virtual uint32 GetCategories() override;
+	virtual class UThumbnailInfo* GetThumbnailInfo(UObject* Asset) const override;
+	virtual bool IsImportedAsset() const override;
+	virtual void GetResolvedSourceFilePaths(const TArray<UObject*>& TypeAssets, TArray<FString>& OutSourceFilePaths) const override;
+	// End IAssetTypeActions
+};
diff --git a/Engine/Source/Editor/GeometryCacheEd/Classes/GeometryCacheAssetBroker.h b/Engine/Source/Editor/GeometryCacheEd/Classes/GeometryCacheAssetBroker.h
new file mode 100644
--- /dev/null
+++ b/Engine/Source/Editor/GeometryCacheEd/Classes/GeometryCacheAssetBroker.h
@@ -0,0 +1,17 @@
+// Copyright 1998-2015 Epic Games, Inc. All Rights Reserved.
+
+#pragma once
+
+#include "GeometryCacheEdModulePublicPCH.h"
+#include "ComponentAssetBroker.h"
+
+/** AssetBroker class for GeometryCache assets*/
+class FGeometryCacheAssetBroker : public IComponentAssetBroker
+{
+public:
+	// Begin IComponentAssetBroker Implementation
+	UClass* GetSupportedAssetClass() override;
+	virtual bool AssignAssetToComponent(UActorComponent* InComponent, UObject* InAsset) override;
+	virtual UObject* GetAssetFromComponent(UActorComponent* InComponent) override;
+	// End IComponentAssetBroker Implementation
+};
diff --git a/Engine/Source/Editor/GeometryCacheEd/Classes/GeometryCacheThumbnailRenderer.h b/Engine/Source/Editor/GeometryCacheEd/Classes/GeometryCacheThumbnailRenderer.h
new file mode 100644
--- /dev/null
+++ b/Engine/Source/Editor/GeometryCacheEd/Classes/GeometryCacheThumbnailRenderer.h
@@ -0,0 +1,27 @@
+// Copyright 1998-2015 Epic Games, Inc. All Rights Reserved.
+
+/**
+*
+* This thumbnail renderer displays a given GeometryCache
+*/
+
+#pragma once
+#include "GeometryCacheThumbnailRenderer.generated.h"
+
+class FGeometryCacheThumbnailScene;
+
+UCLASS(config = Editor, MinimalAPI)
+class UGeometryCacheThumbnailRenderer : public UDefaultSizedThumbnailRenderer
+{
+	GENERATED_UCLASS_BODY()
+
+	// Begin UThumbnailRenderer Object
+	virtual void Draw(UObject* Object, int32 X, int32 Y, uint32 Width, uint32 Height, FRenderTarget* RenderTarget, FCanvas* Canvas) override;
+	// End UThumbnailRenderer Object
+
+	// Begin UObject implementation
+	virtual void BeginDestroy() override;
+	// End UObject implementation
+private:
+	FGeometryCacheThumbnailScene* ThumbnailScene;
+};
diff --git a/Engine/Source/Editor/GeometryCacheEd/Classes/GeometryCacheThumbnailScene.h b/Engine/Source/Editor/GeometryCacheEd/Classes/GeometryCacheThumbnailScene.h
new file mode 100644
--- /dev/null
+++ b/Engine/Source/Editor/GeometryCacheEd/Classes/GeometryCacheThumbnailScene.h
@@ -0,0 +1,26 @@
+// Copyright 1998-2015 Epic Games, Inc. All Rights Reserved.
+
+#pragma once
+
+#include "Thumbnailhelpers.h"
+
+class UGeometryCache;
+class AGeometryCacheActor;
+
+class FGeometryCacheThumbnailScene : public FThumbnailPreviewScene
+{
+public:
+	/** Constructor */
+	FGeometryCacheThumbnailScene();
+
+	/** Sets the static mesh to use in the next GetView() */
+	void SetGeometryCache(UGeometryCache* GeometryCache);
+
+protected:
+	// FThumbnailPreviewScene implementation
+	virtual void GetViewMatrixParameters(const float InFOVDegrees, FVector& OutOrigin, float& OutOrbitPitch, float& OutOrbitYaw, float& OutOrbitZoom) const override;
+
+private:
+	/** The static mesh actor used to display all static mesh thumbnails */
+	AGeometryCacheActor* PreviewActor;
+};
\ No newline at end of file
diff --git a/Engine/Source/Editor/GeometryCacheEd/GeometryCacheEd.Build.cs b/Engine/Source/Editor/GeometryCacheEd/GeometryCacheEd.Build.cs
new file mode 100644
--- /dev/null
+++ b/Engine/Source/Editor/GeometryCacheEd/GeometryCacheEd.Build.cs
@@ -0,0 +1,30 @@
+// Copyright 1998-2015 Epic Games, Inc. All Rights Reserved.
+
+using UnrealBuildTool;
+
+public class GeometryCacheEd : ModuleRules
+{
+	public GeometryCacheEd(TargetInfo Target)
+	{
+        PublicIncludePaths.Add("Runtime/GeometryCacheEd/Classes");
+        PublicIncludePaths.Add("Runtime/GeometryCacheEd/Public");
+        PrivateIncludePaths.Add("Runtime/GeometryCacheEd/Private");
+        
+        PublicDependencyModuleNames.AddRange(
+			new string[] {
+				"Core",
+				"CoreUObject",
+				"Engine",
+				"Slate",
+				"SlateCore",
+                "InputCore",
+                "RenderCore",
+                "ShaderCore",
+                "RHI",		
+                "UnrealEd",
+				"AssetTools",
+                "GeometryCache"
+			}
+		);
+	}
+}
diff --git a/Engine/Source/Editor/GeometryCacheEd/Private/ActorFactoryGeometryCache.cpp b/Engine/Source/Editor/GeometryCacheEd/Private/ActorFactoryGeometryCache.cpp
new file mode 100644
--- /dev/null
+++ b/Engine/Source/Editor/GeometryCacheEd/Private/ActorFactoryGeometryCache.cpp
@@ -0,0 +1,59 @@
+// Copyright 1998-2015 Epic Games, Inc. All Rights Reserved.
+
+#include "GeometryCacheEdModulePublicPCH.h"
+#include "ActorFactoryGeometryCache.h"
+#include "GeometryCache.h"
+#include "GeometryCacheActor.h"
+#include "GeometryCacheComponent.h"
+
+UActorFactoryGeometryCache::UActorFactoryGeometryCache(const FObjectInitializer& ObjectInitializer)
+: Super(ObjectInitializer)
+{
+	DisplayName = FText::FromString( "Geometry Cache" );
+	NewActorClass = AGeometryCacheActor::StaticClass();
+	bUseSurfaceOrientation = true;
+}
+
+bool UActorFactoryGeometryCache::CanCreateActorFrom(const FAssetData& AssetData, FText& OutErrorMsg)
+{
+	if (!AssetData.IsValid() || !AssetData.GetClass()->IsChildOf(UGeometryCache::StaticClass()))
+	{
+		OutErrorMsg = FText::FromString("A valid GeometryCache must be specified.");
+		return false;
+	}
+
+	return true;
+}
+
+void UActorFactoryGeometryCache::PostSpawnActor(UObject* Asset, AActor* NewActor)
+{
+	Super::PostSpawnActor(Asset, NewActor);
+
+	UGeometryCache* GeometryCache = CastChecked<UGeometryCache>(Asset);
+
+	// Change properties
+	AGeometryCacheActor* GeometryCacheActor = CastChecked<AGeometryCacheActor>(NewActor);
+	UGeometryCacheComponent* GeometryCacheComponent = GeometryCacheActor->GetGeometryCacheComponent();
+	check(GeometryCacheComponent);
+
+	GeometryCacheComponent->UnregisterComponent();
+
+	// Set GeometryCache (data) instance
+	GeometryCacheComponent->GeometryCache = GeometryCache;
+
+	// Init Component
+	GeometryCacheComponent->RegisterComponent();
+}
+
+void UActorFactoryGeometryCache::PostCreateBlueprint(UObject* Asset, AActor* CDO)
+{
+	if (Asset != NULL && CDO != NULL)
+	{
+		// Set GeometryCache (data) instance
+		UGeometryCache* GeometryCache = CastChecked<UGeometryCache>(Asset);
+		AGeometryCacheActor* GeometryCacheActor = CastChecked<AGeometryCacheActor>(CDO);
+		UGeometryCacheComponent* GeometryCacheComponent = GeometryCacheActor->GetGeometryCacheComponent();
+
+		GeometryCacheComponent->GeometryCache = GeometryCache;
+	}
+}
\ No newline at end of file
diff --git a/Engine/Source/Editor/GeometryCacheEd/Private/AssetTypeActions_GeometryCache.cpp b/Engine/Source/Editor/GeometryCacheEd/Private/AssetTypeActions_GeometryCache.cpp
new file mode 100644
--- /dev/null
+++ b/Engine/Source/Editor/GeometryCacheEd/Private/AssetTypeActions_GeometryCache.cpp
@@ -0,0 +1,66 @@
+// Copyright 1998-2015 Epic Games, Inc. All Rights Reserved.
+
+#include "GeometryCacheEdModulePublicPCH.h"
+#include "AssetTypeActions_GeometryCache.h"
+
+FText FAssetTypeActions_GeometryCache::GetName() const
+{
+	return NSLOCTEXT("AssetTypeActions", "AssetTypeActions_GeometryCache", "GeometryCache");
+}
+
+FColor FAssetTypeActions_GeometryCache::GetTypeColor() const
+{
+	return FColor(0, 255, 255);
+}
+
+UClass* FAssetTypeActions_GeometryCache::GetSupportedClass() const
+{
+	return UGeometryCache::StaticClass();
+}
+
+bool FAssetTypeActions_GeometryCache::HasActions(const TArray<UObject*>& InObjects) const
+{
+	return false;
+}
+
+void FAssetTypeActions_GeometryCache::GetActions(const TArray<UObject*>& InObjects, FMenuBuilder& MenuBuilder)
+{
+	FAssetTypeActions_Base::GetActions(InObjects, MenuBuilder);
+}
+
+void FAssetTypeActions_GeometryCache::OpenAssetEditor(const TArray<UObject*>& InObjects, TSharedPtr<class IToolkitHost> EditWithinLevelEditor /*= TSharedPtr<IToolkitHost>()*/)
+{
+	FAssetTypeActions_Base::OpenAssetEditor(InObjects, EditWithinLevelEditor);
+}
+
+uint32 FAssetTypeActions_GeometryCache::GetCategories()
+{
+	return EAssetTypeCategories::Basic;
+}
+
+class UThumbnailInfo* FAssetTypeActions_GeometryCache::GetThumbnailInfo(UObject* Asset) const
+{
+	UGeometryCache* GeometryCache = CastChecked<UGeometryCache>(Asset);
+	UThumbnailInfo* ThumbnailInfo = GeometryCache->ThumbnailInfo;
+	if (ThumbnailInfo == NULL)
+	{
+		ThumbnailInfo = NewObject<USceneThumbnailInfo>(GeometryCache);
+		GeometryCache->ThumbnailInfo = ThumbnailInfo;
+	}
+
+	return ThumbnailInfo;
+}
+
+bool FAssetTypeActions_GeometryCache::IsImportedAsset() const
+{
+	return true;
+}
+
+void FAssetTypeActions_GeometryCache::GetResolvedSourceFilePaths(const TArray<UObject*>& TypeAssets, TArray<FString>& OutSourceFilePaths) const
+{
+	for (auto& Asset : TypeAssets)
+	{
+		const auto GeometryCache = CastChecked<UGeometryCache>(Asset);
+		GeometryCache->AssetImportData->ExtractFilenames(OutSourceFilePaths);
+	}
+}
diff --git a/Engine/Source/Editor/GeometryCacheEd/Private/GeometryCacheAssetBroker.cpp b/Engine/Source/Editor/GeometryCacheEd/Private/GeometryCacheAssetBroker.cpp
new file mode 100644
--- /dev/null
+++ b/Engine/Source/Editor/GeometryCacheEd/Private/GeometryCacheAssetBroker.cpp
@@ -0,0 +1,37 @@
+// Copyright 1998-2015 Epic Games, Inc. All Rights Reserved.
+
+#include "GeometryCacheEdModulePublicPCH.h"
+#include "GeometryCacheAssetBroker.h"
+#include "GeometryCacheComponent.h"
+#include "GeometryCache.h"
+
+UClass* FGeometryCacheAssetBroker::GetSupportedAssetClass()
+{
+	return UGeometryCache::StaticClass();
+}
+
+bool FGeometryCacheAssetBroker::AssignAssetToComponent(UActorComponent* InComponent, UObject* InAsset)
+{
+	if (UGeometryCacheComponent* GeometryCacheComponent = Cast<UGeometryCacheComponent>(InComponent))
+	{
+		UGeometryCache* GeomCache = Cast<UGeometryCache>(InAsset);
+
+		if ((GeomCache != nullptr) || (InAsset == nullptr))
+		{
+			GeometryCacheComponent->SetGeometryCache(GeomCache);
+
+			return true;
+		}
+	}
+
+	return false;
+}
+
+UObject* FGeometryCacheAssetBroker::GetAssetFromComponent(UActorComponent* InComponent)
+{
+	if (UGeometryCacheComponent* GeometryCacheComponent = Cast<UGeometryCacheComponent>(InComponent))
+	{
+		return GeometryCacheComponent->GetGeometryCache();
+	}
+	return nullptr;
+}
diff --git a/Engine/Source/Editor/GeometryCacheEd/Private/GeometryCacheEdModule.cpp b/Engine/Source/Editor/GeometryCacheEd/Private/GeometryCacheEdModule.cpp
new file mode 100644
--- /dev/null
+++ b/Engine/Source/Editor/GeometryCacheEd/Private/GeometryCacheEdModule.cpp
@@ -0,0 +1,34 @@
+// Copyright 1998-2015 Epic Games, Inc. All Rights Reserved.
+
+#include "GeometryCacheEdModulePublicPCH.h"
+#include "GeometryCacheEdModule.h"
+#include "IAssetTools.h"
+#include "ComponentAssetBroker.h"
+#include "AssetTypeActions_GeometryCache.h"
+#include "GeometryCacheAssetBroker.h"
+#include "GeometryCacheComponent.h"
+
+IMPLEMENT_MODULE(FGeometryCacheEdModule, GeometryCacheEd)
+
+void FGeometryCacheEdModule::StartupModule()
+{
+	FAssetToolsModule& AssetToolsModule = FAssetToolsModule::GetModule();
+
+	IAssetTools& AssetTools = AssetToolsModule.Get();
+	AssetAction = new FAssetTypeActions_GeometryCache();
+	AssetTools.RegisterAssetTypeActions(MakeShareable(AssetAction));
+
+	AssetBroker = new FGeometryCacheAssetBroker();
+	FComponentAssetBrokerage::RegisterBroker(MakeShareable(AssetBroker), UGeometryCacheComponent::StaticClass(), true, true);
+}
+
+void FGeometryCacheEdModule::ShutdownModule()
+{
+	if (UObjectInitialized())
+	{
+		FAssetToolsModule& AssetToolsModule = FAssetToolsModule::GetModule();
+		IAssetTools& AssetTools = AssetToolsModule.Get();
+		AssetTools.UnregisterAssetTypeActions(AssetAction->AsShared());
+		FComponentAssetBrokerage::UnregisterBroker(MakeShareable(AssetBroker));
+	}
+}
\ No newline at end of file
diff --git a/Engine/Source/Editor/GeometryCacheEd/Private/GeometryCacheThumbnailRenderer.cpp b/Engine/Source/Editor/GeometryCacheEd/Private/GeometryCacheThumbnailRenderer.cpp
new file mode 100644
--- /dev/null
+++ b/Engine/Source/Editor/GeometryCacheEd/Private/GeometryCacheThumbnailRenderer.cpp
@@ -0,0 +1,53 @@
+// Copyright 1998-2015 Epic Games, Inc. All Rights Reserved.
+
+#include "GeometryCacheEdModulePublicPCH.h"
+#include "GeometryCacheThumbnailRenderer.h"
+#include "GeometryCacheThumbnailScene.h"
+#include "Thumbnailhelpers.h"
+#include "EngineModule.h"
+#include "Sceneview.h"
+#include "RendererInterface.h"
+#include "GeometryCache.h"
+
+UGeometryCacheThumbnailRenderer::UGeometryCacheThumbnailRenderer(const FObjectInitializer& ObjectInitializer)
+: Super(ObjectInitializer)
+{
+	ThumbnailScene = nullptr;
+}
+
+void UGeometryCacheThumbnailRenderer::Draw(UObject* Object, int32 X, int32 Y, uint32 Width, uint32 Height, FRenderTarget* RenderTarget, FCanvas* Canvas)
+{
+	UGeometryCache* GeometryCache = Cast<UGeometryCache>(Object);
+	if (GeometryCache != nullptr && !GeometryCache->IsPendingKill())
+	{
+		if (ThumbnailScene == nullptr)
+		{
+			ThumbnailScene = new FGeometryCacheThumbnailScene();
+		}
+
+		ThumbnailScene->SetGeometryCache(GeometryCache);
+		ThumbnailScene->GetScene()->UpdateSpeedTreeWind(0.0);
+
+		FSceneViewFamilyContext ViewFamily(FSceneViewFamily::ConstructionValues(RenderTarget, ThumbnailScene->GetScene(), FEngineShowFlags(ESFIM_Game))
+			.SetWorldTimes(FApp::GetCurrentTime() - GStartTime, FApp::GetDeltaTime(), FApp::GetCurrentTime() - GStartTime));
+
+		ViewFamily.EngineShowFlags.DisableAdvancedFeatures();
+		ViewFamily.EngineShowFlags.MotionBlur = 0;
+		ViewFamily.EngineShowFlags.LOD = 0;
+
+		ThumbnailScene->GetView(&ViewFamily, X, Y, Width, Height);
+		GetRendererModule().BeginRenderingViewFamily(Canvas, &ViewFamily);
+		ThumbnailScene->SetGeometryCache(nullptr);
+	}
+}
+
+void UGeometryCacheThumbnailRenderer::BeginDestroy()
+{
+	if (ThumbnailScene != nullptr)
+	{
+		delete ThumbnailScene;
+		ThumbnailScene = nullptr;
+	}
+
+	Super::BeginDestroy();
+}
diff --git a/Engine/Source/Editor/GeometryCacheEd/Private/GeometryCacheThumbnailScene.cpp b/Engine/Source/Editor/GeometryCacheEd/Private/GeometryCacheThumbnailScene.cpp
new file mode 100644
--- /dev/null
+++ b/Engine/Source/Editor/GeometryCacheEd/Private/GeometryCacheThumbnailScene.cpp
@@ -0,0 +1,68 @@
+// Copyright 1998-2015 Epic Games, Inc. All Rights Reserved.
+
+#include "GeometryCacheEdModulePublicPCH.h"
+#include "GeometryCacheThumbnailRenderer.h"
+#include "GeometryCacheThumbnailScene.h"
+#include "GeometryCacheActor.h"
+#include "GeometryCacheComponent.h"
+
+FGeometryCacheThumbnailScene::FGeometryCacheThumbnailScene()
+{
+	bForceAllUsedMipsResident = false;
+
+	// Create preview actor
+	// checked
+	FActorSpawnParameters SpawnInfo;
+	SpawnInfo.SpawnCollisionHandlingOverride = ESpawnActorCollisionHandlingMethod::AlwaysSpawn;
+	SpawnInfo.bNoFail = true;
+	SpawnInfo.ObjectFlags = RF_Transient;
+	PreviewActor = GetWorld()->SpawnActor<AGeometryCacheActor>(SpawnInfo);
+
+	PreviewActor->GetGeometryCacheComponent()->SetMobility(EComponentMobility::Movable);
+	PreviewActor->SetActorEnableCollision(false);
+}
+
+void FGeometryCacheThumbnailScene::SetGeometryCache(UGeometryCache* GeometryCache)
+{
+	PreviewActor->GetGeometryCacheComponent()->SetGeometryCache(GeometryCache);
+
+	if (GeometryCache)
+	{
+		FTransform MeshTransform = FTransform::Identity;
+
+		PreviewActor->SetActorLocation(FVector(0, 0, 0), false);
+		PreviewActor->GetGeometryCacheComponent()->UpdateBounds();
+
+		// Center the mesh at the world origin then offset to put it on top of the plane
+		const float BoundsZOffset = GetBoundsZOffset(PreviewActor->GetGeometryCacheComponent()->Bounds);
+		PreviewActor->SetActorLocation(-PreviewActor->GetGeometryCacheComponent()->Bounds.Origin + FVector(0, 0, BoundsZOffset), false);
+		PreviewActor->GetGeometryCacheComponent()->RecreateRenderState_Concurrent();
+	}
+}
+
+void FGeometryCacheThumbnailScene::GetViewMatrixParameters(const float InFOVDegrees, FVector& OutOrigin, float& OutOrbitPitch, float& OutOrbitYaw, float& OutOrbitZoom) const
+{
+	check(PreviewActor);
+	check(PreviewActor->GetGeometryCacheComponent());
+	check(PreviewActor->GetGeometryCacheComponent()->GeometryCache);
+
+	const float HalfFOVRadians = FMath::DegreesToRadians<float>(InFOVDegrees) * 0.5f;
+	// Add extra size to view slightly outside of the sphere to compensate for perspective
+	const float HalfMeshSize = PreviewActor->GetGeometryCacheComponent()->Bounds.SphereRadius * 1.15;
+	const float BoundsZOffset = GetBoundsZOffset(PreviewActor->GetGeometryCacheComponent()->Bounds);
+	const float TargetDistance = HalfMeshSize / FMath::Tan(HalfFOVRadians);
+
+	USceneThumbnailInfo* ThumbnailInfo = Cast<USceneThumbnailInfo>(PreviewActor->GetGeometryCacheComponent()->GeometryCache->ThumbnailInfo);
+	if (ThumbnailInfo)
+	{
+		if (TargetDistance + ThumbnailInfo->OrbitZoom < 0)
+		{
+			ThumbnailInfo->OrbitZoom = -TargetDistance;
+		}
+	}
+	else
+	{
+		ThumbnailInfo = USceneThumbnailInfo::StaticClass()->GetDefaultObject<USceneThumbnailInfo>();
+	}
+}
+
diff --git a/Engine/Source/Editor/GeometryCacheEd/Public/GeometryCacheEdModule.h b/Engine/Source/Editor/GeometryCacheEd/Public/GeometryCacheEdModule.h
new file mode 100644
--- /dev/null
+++ b/Engine/Source/Editor/GeometryCacheEd/Public/GeometryCacheEdModule.h
@@ -0,0 +1,41 @@
+// Copyright 1998-2015 Epic Games, Inc. All Rights Reserved.
+
+#pragma once
+#include "ModuleManager.h"
+
+/**
+ * The public interface to this module
+ */
+class FGeometryCacheEdModule : public IModuleInterface
+{
+public:
+	FGeometryCacheEdModule() {}
+
+	virtual void StartupModule();
+	virtual void ShutdownModule();
+
+	/**
+	 * Singleton-like access to this module's interface.  This is just for convenience!
+	 * Beware of calling this during the shutdown phase, though.  Your module might have been unloaded already.
+	 *
+	 * @return Returns singleton instance, loading the module on demand if needed
+	 */
+	static inline FGeometryCacheEdModule& Get()
+	{
+		return FModuleManager::LoadModuleChecked< FGeometryCacheEdModule >("GeometryCacheEd");
+	}
+
+	/**
+	 * Checks to see if this module is loaded and ready.  It is only valid to call Get() if IsAvailable() returns true.
+	 *
+	 * @return True if the module is loaded and ready to use
+	 */
+	static inline bool IsAvailable()
+	{
+		return FModuleManager::Get().IsModuleLoaded( "GeometryCacheEd" );
+	}
+private:
+	class FAssetTypeActions_GeometryCache* AssetAction;
+	class FGeometryCacheAssetBroker* AssetBroker;
+};
+
diff --git a/Engine/Source/Editor/GeometryCacheEd/Public/GeometryCacheEdModulePublicPCH.h b/Engine/Source/Editor/GeometryCacheEd/Public/GeometryCacheEdModulePublicPCH.h
new file mode 100644
--- /dev/null
+++ b/Engine/Source/Editor/GeometryCacheEd/Public/GeometryCacheEdModulePublicPCH.h
@@ -0,0 +1,7 @@
+// Copyright 1998-2015 Epic Games, Inc. All Rights Reserved.
+
+//#include "ActorFactoryGeometryCache.h"
+//#include "AssetTypeActions_GeometryCache.h"
+//#include "GeometryCacheAssetBroker.h"
+//#include "GeometryCacheThumbnailRenderer.h"
+//#include "GeometryCacheThumbnailScene.h"
\ No newline at end of file
diff --git a/Engine/Source/Editor/UnrealEd/UnrealEd.Build.cs b/Engine/Source/Editor/UnrealEd/UnrealEd.Build.cs
--- a/Engine/Source/Editor/UnrealEd/UnrealEd.Build.cs
+++ b/Engine/Source/Editor/UnrealEd/UnrealEd.Build.cs
@@ -136,7 +136,7 @@ public class UnrealEd : ModuleRules
                 "MaterialUtilities",
                 "LocalizationService",
 				"AddContentDialog",
-				"GameProjectGeneration",
+				"GameProjectGeneration"                
 			}
 		);
 
@@ -206,6 +206,7 @@ public class UnrealEd : ModuleRules
 				"SizeMap",
                 "PortalProxies",
                 "PortalServices",
+                "GeometryCacheEd"
 			}
 		);
 
diff --git a/Engine/Source/Runtime/GeometryCache/Classes/GeometryCache.h b/Engine/Source/Runtime/GeometryCache/Classes/GeometryCache.h
--- a/Engine/Source/Runtime/GeometryCache/Classes/GeometryCache.h
+++ b/Engine/Source/Runtime/GeometryCache/Classes/GeometryCache.h
@@ -1,9 +1,11 @@
 // Copyright 1998-2015 Epic Games, Inc. All Rights Reserved.
 #pragma once
-#include "Interfaces/Interface_AssetUserData.h"
+#include "Core.h"
+#include "CoreUObject.h"
+#include "RenderCore.h"
 
+#include "Interfaces/Interface_AssetUserData.h"
 #include "GeometryCacheTrack.h"
-
 #include "GeometryCache.generated.h"
 
 /**
@@ -64,6 +66,6 @@ private:
 	uint32 NumVertexAnimationTracks;
 	uint32 NumTransformAnimationTracks;
 
-	/** A fence which is used to keep track of the rendering thread releasing the static mesh resources. */
+	/** A fence which is used to keep track of the rendering thread releasing the geometry cache resources. */
 	FRenderCommandFence ReleaseResourcesFence;
 };
diff --git a/Engine/Source/Runtime/GeometryCache/Classes/GeometryCacheActor.h b/Engine/Source/Runtime/GeometryCache/Classes/GeometryCacheActor.h
--- a/Engine/Source/Runtime/GeometryCache/Classes/GeometryCacheActor.h
+++ b/Engine/Source/Runtime/GeometryCache/Classes/GeometryCacheActor.h
@@ -15,7 +15,7 @@ class AGeometryCacheActor : public AActor
 
 	// Begin AActor overrides.
 #if WITH_EDITOR
-		virtual bool GetReferencedContentObjects(TArray<UObject*>& Objects) const override;
+	 virtual bool GetReferencedContentObjects(TArray<UObject*>& Objects) const override; 
 #endif // WITH_EDITOR
 	// End AActor overrides.
 
@@ -24,5 +24,5 @@ class AGeometryCacheActor : public AActor
 public:
 	/** Returns GeometryCacheComponent subobject **/
 	UFUNCTION(BlueprintCallable, Category = "Components|GeometryCache")
-	UGeometryCacheComponent* GetGeometryCacheComponent() const;	
+	GEOMETRYCACHE_API UGeometryCacheComponent* GetGeometryCacheComponent() const;
 };
diff --git a/Engine/Source/Runtime/GeometryCache/Classes/GeometryCacheComponent.h b/Engine/Source/Runtime/GeometryCache/Classes/GeometryCacheComponent.h
--- a/Engine/Source/Runtime/GeometryCache/Classes/GeometryCacheComponent.h
+++ b/Engine/Source/Runtime/GeometryCache/Classes/GeometryCacheComponent.h
@@ -2,7 +2,7 @@
 
 #pragma once
 
-#include "GeometryCacheModulePrivatePCH.h"
+#include "GeometryCacheModulePublicPCH.h"
 #include "GeometryCache.h"
 #include "DynamicMeshBuilder.h"
 
@@ -19,20 +19,17 @@ struct FTrackRenderData
 
 	~FTrackRenderData()
 	{
-		MeshData = NULL;
+		MeshData = nullptr;
 	}
 	
-	/** Pointer to FGeometryCacheMeshData containing vertex-data, bounding box and batch-info*/
+	/** Pointer to FGeometryCacheMeshData containing vertex-data, bounding box, index buffer and batch-info*/
 	FGeometryCacheMeshData* MeshData;
 	/** World matrix used to render this specific track */
 	FMatrix WorldMatrix;
-	/** Index Buffer generated according to the vertex-data */
-	TArray<int32> IndexBuffer;
 
 	void Reset()
 	{
-		IndexBuffer.Empty();
-		MeshData = NULL;
+		MeshData = nullptr;
 	}
 };
 
@@ -112,15 +109,27 @@ class GEOMETRYCACHE_API UGeometryCacheComponent : public UMeshComponent
 	* @param SectionIndex - Index of the section we want to update
 	* @param Indices - Array with new indices
 	*/
-	void UpdateTrackSectionIndexbuffer(int32 SectionIndex, const TArray<int32>& Indices );
+	void UpdateTrackSectionIndexbuffer(int32 SectionIndex, const TArray<uint32>& Indices );
 
 	/**
-	* OnObjectReimported QQQ
+	* OnObjectReimported, Callback function to refresh section data and update scene proxy.
 	*
-	* @param ImportedGeometryCache - Callback function to refresh section data and update sceneproxy
+	* @param ImportedGeometryCache
 	*/
 	void OnObjectReimported(UGeometryCache* ImportedGeometryCache);
 
+	/**
+	* SetupTrackData
+	* Setup data required for playback of geometry cache tracks
+	*/
+	void SetupTrackData();
+	
+	/**
+	* ClearTrackData
+	* Clean up data that was required for playback of geometry cache tracks
+	*/
+	void ClearTrackData();
+
 	/** Start playback of GeometryCache */
 	UFUNCTION(BlueprintCallable, Category = "Components|GeometryCache")
 	void Play();
@@ -168,11 +177,29 @@ class GEOMETRYCACHE_API UGeometryCacheComponent : public UMeshComponent
 	/** Get whether this GeometryCache is playing or not. */
 	UFUNCTION(BlueprintCallable, Category = "Components|GeometryCache")
 	void SetPlaybackSpeed(const float NewPlaybackSpeed);
-	
+
+	/** Change the Geometry Cache used by this instance. */
+	UFUNCTION(BlueprintCallable, Category = "Components|GeometryCache")
+	virtual bool SetGeometryCache( UGeometryCache* NewGeomCache );
+
+	/** Getter for Geometry cache instance referred by the component */
+	UGeometryCache* GetGeometryCache() const;
+
 	/** Geometry Cache instance referenced by the component */
 	UPROPERTY(EditAnywhere, BlueprintReadOnly, Category = GeometryCache)
 	UGeometryCache* GeometryCache;
-	
+		
+protected:
+	/**
+	* Invalidate both the Matrix and Mesh sample indices
+	*/
+	void InvalidateTrackSampleIndices();
+
+	/**
+	* ReleaseResources, clears and removes data stored/copied from GeometryCache instance	
+	*/
+	void ReleaseResources();	
+
 	UPROPERTY(EditAnywhere, Category = GeometryCache)
 	bool bRunning;
 
@@ -191,17 +218,6 @@ class GEOMETRYCACHE_API UGeometryCacheComponent : public UMeshComponent
 	UPROPERTY(VisibleAnywhere, Category = GeometryCache)
 	float ElapsedTime;
 
-protected:
-	/**
-	* Invalidate both the Matrix and Mesh sample indices
-	*/
-	void InvalidateTrackSampleIndices();
-
-	/**
-	* ReleaseResources, clears and removes data stored/copied from GeometryCache instance	
-	*/
-	void ReleaseResources();	
-
 	/** Local space bounds of mesh */
 	FBoxSphereBounds LocalBounds;
 
diff --git a/Engine/Source/Runtime/GeometryCache/Classes/GeometryCacheMeshData.h b/Engine/Source/Runtime/GeometryCache/Classes/GeometryCacheMeshData.h
--- a/Engine/Source/Runtime/GeometryCache/Classes/GeometryCacheMeshData.h
+++ b/Engine/Source/Runtime/GeometryCache/Classes/GeometryCacheMeshData.h
@@ -48,7 +48,9 @@ struct FGeometryCacheMeshData
 	TArray<FGeometryCacheMeshBatchInfo> BatchesInfo;
 	/** Bounding box for this sample in the track */
 	FBox BoundingBox;
-	
+	/** Indices for this sample, used for drawing the mesh */
+	TArray<uint32> Indices;
+		
 	/** Serialization for FVertexAnimationSample. */
 	friend FArchive& operator<<(FArchive& Ar, FGeometryCacheMeshData& Mesh)
 	{
@@ -78,6 +80,8 @@ struct FGeometryCacheMeshData
 		Ar << Mesh.BoundingBox;
 		Ar << Mesh.BatchesInfo;
 
+		Ar << Mesh.Indices;	
+
 		return Ar;
 	}
 
@@ -90,6 +94,8 @@ struct FGeometryCacheMeshData
 		Size += sizeof(Vertices);
 		Size += sizeof(BatchesInfo);
 		Size += sizeof(BoundingBox);
+		Size += Indices.Num() * sizeof(uint32);
+		Size += sizeof(Indices);
 
 		return Size;
 	}
diff --git a/Engine/Source/Runtime/GeometryCache/Classes/GeometryCacheTrackTransformGroupAnimation.h b/Engine/Source/Runtime/GeometryCache/Classes/GeometryCacheTrackTransformGroupAnimation.h
new file mode 100644
--- /dev/null
+++ b/Engine/Source/Runtime/GeometryCache/Classes/GeometryCacheTrackTransformGroupAnimation.h
@@ -0,0 +1,32 @@
+// Copyright 1998-2015 Epic Games, Inc. All Rights Reserved.
+#pragma once
+
+#include "GeometryCacheMeshData.h"
+
+#include "GeometryCacheTrackTransformGroupAnimation.generated.h"
+
+/** Derived GeometryCacheTrack class, used for Transform animation. */
+UCLASS(collapsecategories, hidecategories = Object, BlueprintType, config = Engine)
+class GEOMETRYCACHE_API UGeometryCacheTrack_TransformGroupAnimation : public UGeometryCacheTrack
+{
+	GENERATED_UCLASS_BODY()
+		
+	// Begin UObject interface.
+	virtual SIZE_T GetResourceSize(EResourceSizeMode::Type Mode) override;
+	virtual void Serialize(FArchive& Ar) override;
+	// End UObject interface.
+	
+	// Begin UGeometryCacheTrack interface.
+	virtual const bool UpdateMeshData(const float Time, const bool bLooping, int32& InOutMeshSampleIndex, FGeometryCacheMeshData*& OutMeshData) override;
+	// End UGeometryCacheTrack interface.
+
+	/**
+	* Sets/updates the MeshData for this track
+	*
+	* @param NewMeshData - GeometryCacheMeshData instance later used as the rendered mesh	
+	*/
+	UFUNCTION()
+	void SetMesh(const FGeometryCacheMeshData& NewMeshData);	
+private:
+	FGeometryCacheMeshData MeshData;
+};
diff --git a/Engine/Source/Runtime/GeometryCache/GeometryCache.Build.cs b/Engine/Source/Runtime/GeometryCache/GeometryCache.Build.cs
--- a/Engine/Source/Runtime/GeometryCache/GeometryCache.Build.cs
+++ b/Engine/Source/Runtime/GeometryCache/GeometryCache.Build.cs
@@ -7,23 +7,28 @@ public class GeometryCache : ModuleRules
 	public GeometryCache(TargetInfo Target)
 	{
         PublicIncludePaths.Add("Runtime/GeometryCache/Public");
+        PublicIncludePaths.Add("Runtime/GeometryCache/Classes");
         PrivateIncludePaths.Add("Runtime/GeometryCache/Private");
         
         PublicDependencyModuleNames.AddRange(
 			new string[] {
 				"Core",
-				"CoreUObject", // @todo Mac: for some reason it's needed to link in debug on Mac
+				"CoreUObject",
 				"Engine",
 				"Slate",
 				"SlateCore",
                 "InputCore",
-                "EditorStyle",
-                "UnrealEd",
-                "AssetTools",
                 "RenderCore",
                 "ShaderCore",
                 "RHI"
 			}
 		);
+
+        PublicIncludePathModuleNames.Add("GeometryCacheEd");
+
+        if (UEBuildConfiguration.bBuildEditor)
+        {
+            DynamicallyLoadedModuleNames.Add("GeometryCacheEd");
+        }
 	}
 }
diff --git a/Engine/Source/Runtime/GeometryCache/Private/GeometryCache.cpp b/Engine/Source/Runtime/GeometryCache/Private/GeometryCache.cpp
--- a/Engine/Source/Runtime/GeometryCache/Private/GeometryCache.cpp
+++ b/Engine/Source/Runtime/GeometryCache/Private/GeometryCache.cpp
@@ -3,6 +3,7 @@
 #include "GeometryCacheModulePrivatePCH.h"
 #include "GeometryCacheTrackTransformAnimation.h"
 #include "GeometryCacheTrackFlipbookAnimation.h"
+#include "EditorFrameWork/AssetImportData.h"
 
 UGeometryCache::UGeometryCache(const FObjectInitializer& ObjectInitializer /*= FObjectInitializer::Get()*/) : UObject(ObjectInitializer)
 {
diff --git a/Engine/Source/Runtime/GeometryCache/Private/GeometryCacheActor.cpp b/Engine/Source/Runtime/GeometryCache/Private/GeometryCacheActor.cpp
--- a/Engine/Source/Runtime/GeometryCache/Private/GeometryCacheActor.cpp
+++ b/Engine/Source/Runtime/GeometryCache/Private/GeometryCacheActor.cpp
@@ -12,7 +12,7 @@ AGeometryCacheActor::AGeometryCacheActor(const FObjectInitializer& ObjectInitial
 	RootComponent = GeometryCacheComponent;
 }
 
-class UGeometryCacheComponent* AGeometryCacheActor::GetGeometryCacheComponent() const
+GEOMETRYCACHE_API UGeometryCacheComponent* AGeometryCacheActor::GetGeometryCacheComponent() const
 {
 	return GeometryCacheComponent;
 }
diff --git a/Engine/Source/Runtime/GeometryCache/Private/GeometryCacheComponent.cpp b/Engine/Source/Runtime/GeometryCache/Private/GeometryCacheComponent.cpp
--- a/Engine/Source/Runtime/GeometryCache/Private/GeometryCacheComponent.cpp
+++ b/Engine/Source/Runtime/GeometryCache/Private/GeometryCacheComponent.cpp
@@ -2,11 +2,15 @@
 
 #include "GeometryCacheModulePrivatePCH.h"
 #include "GeometryCacheComponent.h"
+#include "MessageLog.h"
+#include "ContentStreaming.h"
 
 #include "GeometryCacheSceneProxy.h"
 #include "GeometryCacheTrack.h"
 #include "GeometryCacheTrackFlipbookAnimation.h"
 
+#define LOCTEXT_NAMESPACE "GeometryCacheComponent"
+
 DECLARE_CYCLE_STAT(TEXT("GeometryCacheTick"), STAT_GeometryCacheComponent_TickComponent, STATGROUP_GeometryCache);
 
 UGeometryCacheComponent::UGeometryCacheComponent(const FObjectInitializer& ObjectInitializer)
@@ -29,28 +33,39 @@ void UGeometryCacheComponent::BeginDestroy()
 }
 
 void UGeometryCacheComponent::OnRegister()
+{
+	ClearTrackData();
+	SetupTrackData();
+	Super::OnRegister();
+}
+
+void UGeometryCacheComponent::ClearTrackData()
 {
 	NumTracks = 0;
 	TrackMatrixSampleIndices.Empty();
 	TrackMatrixSampleIndices.Empty();
 	TrackSections.Empty();
+	SceneProxy = nullptr;
+}
 
-	if (GeometryCache != NULL)
+void UGeometryCacheComponent::SetupTrackData()
+{
+	if (GeometryCache != nullptr)
 	{
 		// Refresh NumTracks and clear Index Arrays
-		NumTracks = GeometryCache->Tracks.Num();		
+		NumTracks = GeometryCache->Tracks.Num();
 		TrackMeshSampleIndices.Empty(GeometryCache->Tracks.Num());
 		TrackMatrixSampleIndices.Empty(GeometryCache->Tracks.Num());
 
 		Duration = 0.0f;
 		// Create mesh sections for each GeometryCacheTrack
-		for (int32 TrackIndex = 0; TrackIndex < NumTracks; ++TrackIndex )
-		{ 				
-			UGeometryCacheTrack* Track = GeometryCache->Tracks[TrackIndex];			
+		for (int32 TrackIndex = 0; TrackIndex < NumTracks; ++TrackIndex)
+		{
+			UGeometryCacheTrack* Track = GeometryCache->Tracks[TrackIndex];
 			FMatrix WorldMatrix;
 			int32 MeshSampleIndex = -1;
 			int32 MatrixSampleIndex = -1;
-			FGeometryCacheMeshData* MeshData = NULL;
+			FGeometryCacheMeshData* MeshData = nullptr;
 
 			// Retrieve the matrix/mesh data and the appropriate sample indices
 			Track->UpdateMatrixData(ElapsedTime + StartTimeOffset, bLooping, MatrixSampleIndex, WorldMatrix);
@@ -64,11 +79,9 @@ void UGeometryCacheComponent::OnRegister()
 			TrackMatrixSampleIndices.Add(MatrixSampleIndex);
 
 			const float TrackMaxSampleTime = Track->GetMaxSampleTime();
-			Duration = ( Duration > TrackMaxSampleTime ) ? Duration : TrackMaxSampleTime;
+			Duration = (Duration > TrackMaxSampleTime) ? Duration : TrackMaxSampleTime;
 		}
 	}
-
-	Super::OnRegister();
 }
 
 void UGeometryCacheComponent::OnUnregister()
@@ -89,7 +102,6 @@ void UGeometryCacheComponent::TickComponent(float DeltaTime, enum ELevelTick Tic
 		// Increase total elapsed time since BeginPlay according to PlayDirection and speed
 		ElapsedTime += (DeltaTime * PlayDirection * PlaybackSpeed);
 
-		// QQ change this for reverse playing? extra float for runningtime vs sample time?
 		if (ElapsedTime < 0.0f)
 		{
 			ElapsedTime += Duration;
@@ -100,7 +112,7 @@ void UGeometryCacheComponent::TickComponent(float DeltaTime, enum ELevelTick Tic
 			// Determine if and which part of the section we have to update
 			UGeometryCacheTrack* Track = GeometryCache->Tracks[TrackIndex];
 			FMatrix WorldMatrix;
-			FGeometryCacheMeshData* MeshData = NULL;
+			FGeometryCacheMeshData* MeshData = nullptr;
 
 			const bool bUpdateMatrix = Track->UpdateMatrixData(ElapsedTime + StartTimeOffset, bLooping, TrackMatrixSampleIndices[TrackIndex], WorldMatrix);
 			const bool bUpdateMesh = Track->UpdateMeshData(ElapsedTime + StartTimeOffset, bLooping, TrackMeshSampleIndices[TrackIndex], MeshData);
@@ -180,16 +192,7 @@ void UGeometryCacheComponent::CreateTrackSection(int32 SectionIndex, const FMatr
 
 	// Number of Vertices
 	const int32 NumVerts = MeshData->Vertices.Num();
-
-	// Copy index buffer (clamping to vertex range)
-	const int32 NumTriIndices = NumVerts;
-	NewSection.IndexBuffer.Reset();
-	NewSection.IndexBuffer.AddUninitialized(NumTriIndices);
-	for (int32 IndexIdx = 0; IndexIdx < NumTriIndices; IndexIdx++)
-	{
-		NewSection.IndexBuffer[IndexIdx] = FMath::Min(IndexIdx, NumVerts - 1);
-	}
-	
+		
 	// Store Matrix and MeshData-pointer
 	NewSection.WorldMatrix = WorldMatrix;
 	NewSection.MeshData = MeshData;
@@ -207,43 +210,28 @@ void UGeometryCacheComponent::UpdateTrackSectionMeshData(int32 SectionIndex, FGe
 	
 	// Number of Vertices
 	const int32 NumVerts = MeshData->Vertices.Num();
-	// Determine if we actually need a new index buffer
-	if (UpdateSection.IndexBuffer.Num() != MeshData->Vertices.Num())
-	{
-		// Copy index buffer (clamping to vertex range)
-		const int32 NumTriIndices = NumVerts;
-		UpdateSection.IndexBuffer.Reset();
-		UpdateSection.IndexBuffer.AddUninitialized(NumTriIndices);
-		for (int32 IndexIdx = 0; IndexIdx < NumTriIndices; IndexIdx++)
-		{
-			UpdateSection.IndexBuffer[IndexIdx] = FMath::Min(IndexIdx, NumVerts - 1);
-		}
-
-		if (SceneProxy)
-		{
-			UpdateTrackSectionIndexbuffer(SectionIndex, UpdateSection.IndexBuffer);
-		}
-	}
-
+	
 	// Update MeshDataPointer
 	UpdateSection.MeshData = MeshData;
 
-	if (SceneProxy)
-	{		
-		UpdateTrackSectionVertexbuffer(SectionIndex, MeshData);
-		// Will update the bounds (used for frustum culling)
-		MarkRenderTransformDirty();
-	}
-	else
-	{
-		// Recreate scene proxy
-		MarkRenderStateDirty(); 
-	}
-
-	MarkRenderStateDirty();
+	//if (SceneProxy)
+	//{		
+	//	UpdateTrackSectionVertexbuffer(SectionIndex, MeshData);
+	//	UpdateTrackSectionIndexbuffer(SectionIndex, MeshData->Indices);
+	//	// Will update the bounds (used for frustum culling)
+	//	MarkRenderTransformDirty();
+	//	MarkRenderDynamicDataDirty();
+	//}
+	//else
+	//{
+	//	// Recreate scene proxy
+	//	MarkRenderStateDirty(); 
+	//}
 
 	// Update overall bounds	
-	UpdateLocalBounds(); 	
+	UpdateLocalBounds();
+
+	MarkRenderStateDirty();
 }
 
 void UGeometryCacheComponent::UpdateTrackSectionMatrixData(int32 SectionIndex, const FMatrix& WorldMatrix)
@@ -255,7 +243,7 @@ void UGeometryCacheComponent::UpdateTrackSectionMatrixData(int32 SectionIndex, c
 	// Update the WorldMatrix only
 	UpdateSection.WorldMatrix = WorldMatrix;
 
-	if (!IsRenderStateDirty() && SceneProxy != NULL)
+	if (!IsRenderStateDirty() && SceneProxy != nullptr)
 	{
 		// Update it within the scene proxy
 		SceneProxy->UpdateSectionWorldMatrix(SectionIndex, WorldMatrix);
@@ -280,13 +268,13 @@ void UGeometryCacheComponent::UpdateTrackSectionVertexbuffer(int32 SectionIndex,
 	});
 }
 
-void UGeometryCacheComponent::UpdateTrackSectionIndexbuffer(int32 SectionIndex, const TArray<int32>& Indices)
+void UGeometryCacheComponent::UpdateTrackSectionIndexbuffer(int32 SectionIndex, const TArray<uint32>& Indices)
 {
 	ENQUEUE_UNIQUE_RENDER_COMMAND_THREEPARAMETER(
 		FUpdateVertexBufferCommand,
 		FGeometryCacheSceneProxy*, SceneProxy, SceneProxy,
 		const int32, Index, SectionIndex,
-		const TArray<int32>&, Data, Indices,
+		const TArray<uint32>&, Data, Indices,
 		{
 		SceneProxy->UpdateSectionIndexBuffer(Index, Data);
 	});
@@ -299,7 +287,7 @@ void UGeometryCacheComponent::OnObjectReimported(UGeometryCache* ImportedGeometr
 		ReleaseResources();
 		GeometryCache = ImportedGeometryCache;
 
-		if (GeometryCache != NULL)
+		if (GeometryCache != nullptr)
 		{
 			// Refresh NumTracks and clear Index Arrays
 			NumTracks = GeometryCache->Tracks.Num();
@@ -314,7 +302,7 @@ void UGeometryCacheComponent::OnObjectReimported(UGeometryCache* ImportedGeometr
 				FMatrix WorldMatrix;
 				int32 MeshSampleIndex = -1;
 				int32 MatrixSampleIndex = -1;
-				FGeometryCacheMeshData* MeshData = NULL;
+				FGeometryCacheMeshData* MeshData = nullptr;
 
 				// Retrieve the matrix/mesh data and the appropriate sample indices
 				Track->UpdateMatrixData(ElapsedTime + StartTimeOffset, bLooping, MatrixSampleIndex, WorldMatrix);
@@ -396,6 +384,49 @@ void UGeometryCacheComponent::SetPlaybackSpeed(const float NewPlaybackSpeed)
 	PlaybackSpeed = fabs( NewPlaybackSpeed );
 }
 
+bool UGeometryCacheComponent::SetGeometryCache(UGeometryCache* NewGeomCache)
+{
+	// Do nothing if we are already using the supplied static mesh
+	if (NewGeomCache == GeometryCache)
+	{
+		return false;
+	}
+
+	// Don't allow changing static meshes if "static" and registered
+	AActor* Owner = GetOwner();
+	if (!AreDynamicDataChangesAllowed() && Owner != nullptr)
+	{
+		FMessageLog("PIE").Warning(FText::Format(LOCTEXT("SetGeometryCache", "Calling SetGeometryCache on '{0}' but Mobility is Static."),
+			FText::FromString(GetPathName())));
+		return false;
+	}
+
+
+	GeometryCache = NewGeomCache;
+
+	ClearTrackData();
+	SetupTrackData();
+
+	// Need to send this to render thread at some point
+	MarkRenderStateDirty();
+
+	// Update physics representation right away
+	RecreatePhysicsState();
+	
+	// Notify the streaming system. Don't use Update(), because this may be the first time the geometry cache has been set
+	// and the component may have to be added to the streaming system for the first time.
+	IStreamingManager::Get().NotifyPrimitiveAttached(this, DPT_Spawned);
+
+	// Since we have new tracks, we need to update bounds
+	UpdateBounds();
+	return true;
+}
+
+UGeometryCache* UGeometryCacheComponent::GetGeometryCache() const
+{
+	return GeometryCache;
+}
+
 void UGeometryCacheComponent::PlayReversedFromEnd()
 {
 	ElapsedTime = Duration;
@@ -424,7 +455,7 @@ void UGeometryCacheComponent::InvalidateTrackSampleIndices()
 
 void UGeometryCacheComponent::ReleaseResources()
 {
-	GeometryCache = NULL;
+	GeometryCache = nullptr;
 	NumTracks = 0;
 	TrackMatrixSampleIndices.Empty();
 	TrackMatrixSampleIndices.Empty();
@@ -443,3 +474,5 @@ void UGeometryCacheComponent::PostEditUndo()
 	InvalidateTrackSampleIndices();
 	MarkRenderStateDirty();
 }
+
+#undef LOCTEXT_NAMESPACE
diff --git a/Engine/Source/Runtime/GeometryCache/Private/GeometryCacheModule.cpp b/Engine/Source/Runtime/GeometryCache/Private/GeometryCacheModule.cpp
--- a/Engine/Source/Runtime/GeometryCache/Private/GeometryCacheModule.cpp
+++ b/Engine/Source/Runtime/GeometryCache/Private/GeometryCacheModule.cpp
@@ -1,42 +1,20 @@
 // Copyright 1998-2015 Epic Games, Inc. All Rights Reserved.
 
 #include "GeometryCacheModulePrivatePCH.h"
-#include "AssetToolsModule.h"
-#include "AssetTypeActions_GeometryCache.h"
-
 #include "GeometryCacheModule.h"
 #include "GeometryCacheComponent.h"
+#include "ModuleManager.h"
+#include "GeometryCacheEdModule.h"
 
 IMPLEMENT_MODULE(FGeometryCacheModule, GeometryCache)
 
 void FGeometryCacheModule::StartupModule()
 {
-	// This code will execute after your module is loaded into memory (but after global variables are initialized, of course.)
-
-	FAssetToolsModule& AssetToolsModule = FAssetToolsModule::GetModule();
-
-	IAssetTools& AssetTools = AssetToolsModule.Get();
-
-	AssetTools.RegisterAssetTypeActions(MakeShareable(new FAssetTypeActions_GeometryCache));
-
-	// Register to be notified when an asset is reimported
-	OnAssetReimportDelegateHandle = FEditorDelegates::OnAssetReimport.AddRaw(this, &FGeometryCacheModule::OnObjectReimported);
-
+#if WITH_EDITOR
+	FGeometryCacheEdModule& Module = FModuleManager::LoadModuleChecked<FGeometryCacheEdModule>(TEXT("GeometryCacheEd"));
+#endif
 }
 
 void FGeometryCacheModule::ShutdownModule()
 {
-	// This function may be called during shutdown to clean up your module.  For modules that support dynamic reloading,
-	// we call this function before unloading the module.
-}
-
-void FGeometryCacheModule::OnObjectReimported(UObject* InObject)
-{
-	if (UGeometryCache* GeometryCache = Cast<UGeometryCache>(InObject))
-	{
-		for (TObjectIterator<UGeometryCacheComponent> CacheIt; CacheIt; ++CacheIt)
-		{
-			CacheIt->OnObjectReimported(GeometryCache);
-		}
-	}
-}
+}
\ No newline at end of file
diff --git a/Engine/Source/Runtime/GeometryCache/Private/GeometryCacheModulePrivatePCH.h b/Engine/Source/Runtime/GeometryCache/Private/GeometryCacheModulePrivatePCH.h
--- a/Engine/Source/Runtime/GeometryCache/Private/GeometryCacheModulePrivatePCH.h
+++ b/Engine/Source/Runtime/GeometryCache/Private/GeometryCacheModulePrivatePCH.h
@@ -4,4 +4,7 @@
 
 // You should place include statements to your module's private header files here.  You only need to
 // add includes for headers that are used in most of your module's source files though.
+#include "Core.h"
+#include "Engine.h"
 #include "GeometryCache.h"
+#include "GeometryCacheComponent.h"
diff --git a/Engine/Source/Runtime/GeometryCache/Private/GeometryCacheSceneProxy.cpp b/Engine/Source/Runtime/GeometryCache/Private/GeometryCacheSceneProxy.cpp
--- a/Engine/Source/Runtime/GeometryCache/Private/GeometryCacheSceneProxy.cpp
+++ b/Engine/Source/Runtime/GeometryCache/Private/GeometryCacheSceneProxy.cpp
@@ -15,7 +15,7 @@ FGeometryCacheSceneProxy::FGeometryCacheSceneProxy(UGeometryCacheComponent* Comp
 	{
 		FTrackRenderData& SrcSection = Component->TrackSections[SectionIdx];
 		
-		if (SrcSection.IndexBuffer.Num() > 0)
+		if (SrcSection.MeshData->Indices.Num() > 0)
 		{
 			FGeomCacheTrackProxy* NewSection = new FGeomCacheTrackProxy();
 
@@ -31,7 +31,7 @@ FGeometryCacheSceneProxy::FGeometryCacheSceneProxy(UGeometryCacheComponent* Comp
 			NewSection->VertexBuffer.Vertices.Append(MeshData->Vertices);
 
 			// Copy index buffer
-			NewSection->IndexBuffer.Indices = SrcSection.IndexBuffer;
+			NewSection->IndexBuffer.Indices = SrcSection.MeshData->Indices;
 
 			// Init vertex factory
 			NewSection->VertexFactory.Init(&NewSection->VertexBuffer);
@@ -45,7 +45,7 @@ FGeometryCacheSceneProxy::FGeometryCacheSceneProxy(UGeometryCacheComponent* Comp
 			for (FGeometryCacheMeshBatchInfo& BatchInfo : MeshData->BatchesInfo)
 			{
 				UMaterialInterface* Material = Component->GetMaterial(BatchInfo.MaterialIndex);
-				if (Material == NULL)
+				if (Material == nullptr)
 				{
 					Material = UMaterial::GetDefaultMaterial(MD_Surface);
 				}
@@ -81,11 +81,11 @@ void FGeometryCacheSceneProxy::GetDynamicMeshElements(const TArray<const FSceneV
 	// Set up wireframe material (if needed)
 	const bool bWireframe = AllowDebugViewmodes() && ViewFamily.EngineShowFlags.Wireframe;
 
-	FColoredMaterialRenderProxy* WireframeMaterialInstance = NULL;
+	FColoredMaterialRenderProxy* WireframeMaterialInstance = nullptr;
 	if (bWireframe)
 	{
 		WireframeMaterialInstance = new FColoredMaterialRenderProxy(
-			GEngine->WireframeMaterial ? GEngine->WireframeMaterial->GetRenderProxy(IsSelected()) : NULL,
+			GEngine->WireframeMaterial ? GEngine->WireframeMaterial->GetRenderProxy(IsSelected()) : nullptr,
 			FLinearColor(0, 0.5f, 1.f)
 			);
 
@@ -95,7 +95,7 @@ void FGeometryCacheSceneProxy::GetDynamicMeshElements(const TArray<const FSceneV
 	// Iterate over sections	
 	for (const FGeomCacheTrackProxy* TrackProxy : Sections )
 	{
-		// QQQ
+		// Render out stored TrackProxy's
 		if (TrackProxy != nullptr)
 		{
 			INC_DWORD_STAT_BY(STAT_GeometryCacheSceneProxy_MeshBatchCount, TrackProxy->MeshData->BatchesInfo.Num());
@@ -204,7 +204,7 @@ void FGeometryCacheSceneProxy::UpdateSectionVertexBuffer(const int32 SectionInde
 	}
 }
 
-void FGeometryCacheSceneProxy::UpdateSectionIndexBuffer(const int32 SectionIndex, const TArray<int32>& Indices)
+void FGeometryCacheSceneProxy::UpdateSectionIndexBuffer(const int32 SectionIndex, const TArray<uint32>& Indices)
 {
 	check(SectionIndex < Sections.Num() && "Section Index out of range");
 	check(IsInRenderingThread());
@@ -271,18 +271,20 @@ void FGeomCacheIndexBuffer::InitRHI()
 {
 	FRHIResourceCreateInfo CreateInfo;
 	void* Buffer = nullptr;
-	IndexBufferRHI = RHICreateAndLockIndexBuffer(sizeof(int32), Indices.Num() * sizeof(int32), BUF_Static, CreateInfo, Buffer);
+	IndexBufferRHI = RHICreateAndLockIndexBuffer(sizeof(uint32), Indices.Num() * sizeof(uint32), BUF_Static, CreateInfo, Buffer);
 
 	// Write the indices to the index buffer.	
-	FMemory::Memcpy(Buffer, Indices.GetData(), Indices.Num() * sizeof(int32));
+
+	// Write the indices to the index buffer.
+	FMemory::Memcpy(Buffer, Indices.GetData(), Indices.Num() * sizeof(uint32));
 	RHIUnlockIndexBuffer(IndexBufferRHI);
 }
 
 void FGeomCacheIndexBuffer::UpdateRHI()
 {
 	// Copy the index data into the index buffer.
-	void* Buffer = RHILockIndexBuffer(IndexBufferRHI, 0, Indices.Num() * sizeof(int32), RLM_WriteOnly);
-	FMemory::Memcpy(Buffer, Indices.GetData(), Indices.Num() * sizeof(int32));
+	void* Buffer = RHILockIndexBuffer(IndexBufferRHI, 0, Indices.Num() * sizeof(uint32), RLM_WriteOnly);
+	FMemory::Memcpy(Buffer, Indices.GetData(), Indices.Num() * sizeof(uint32));
 	RHIUnlockIndexBuffer(IndexBufferRHI);
 }
 
diff --git a/Engine/Source/Runtime/GeometryCache/Private/GeometryCacheSceneProxy.h b/Engine/Source/Runtime/GeometryCache/Private/GeometryCacheSceneProxy.h
--- a/Engine/Source/Runtime/GeometryCache/Private/GeometryCacheSceneProxy.h
+++ b/Engine/Source/Runtime/GeometryCache/Private/GeometryCacheSceneProxy.h
@@ -42,7 +42,7 @@ public:
 class GEOMETRYCACHE_API FGeomCacheIndexBuffer : public FIndexBuffer
 {
 public:
-	TArray<int32> Indices;
+	TArray<uint32> Indices;
 
 	virtual void InitRHI() override;
 
@@ -85,27 +85,26 @@ public:
 class GEOMETRYCACHE_API FGeometryCacheSceneProxy : public FPrimitiveSceneProxy
 {
 public:
-	// QQ interfaces/comments
+	
 	FGeometryCacheSceneProxy(class UGeometryCacheComponent* Component);
 
 	virtual ~FGeometryCacheSceneProxy();
 
+	// Begin FPrimitiveSceneProxy interface.
 	virtual void GetDynamicMeshElements(const TArray<const FSceneView*>& Views, const FSceneViewFamily& ViewFamily, uint32 VisibilityMap, FMeshElementCollector& Collector) const override;
-
 	virtual FPrimitiveViewRelevance GetViewRelevance(const FSceneView* View) const override;
 
 	virtual bool CanBeOccluded() const override;
-
 	virtual uint32 GetMemoryFootprint(void) const;
-
 	uint32 GetAllocatedSize(void) const;
+	// End FPrimitiveSceneProxy interface.
 
 	/** Update world matrix for specific section */
 	void UpdateSectionWorldMatrix(const int32 SectionIndex, const FMatrix& WorldMatrix);
 	/** Update vertex buffer for specific section */
 	void UpdateSectionVertexBuffer(const int32 SectionIndex, FGeometryCacheMeshData* MeshData );
 	/** Update index buffer for specific section */
-	void UpdateSectionIndexBuffer(const int32 SectionIndex, const TArray<int32>& Indices);
+	void UpdateSectionIndexBuffer(const int32 SectionIndex, const TArray<uint32>& Indices);
 
 	/** Clears the Sections array*/
 	void ClearSections();
diff --git a/Engine/Source/Runtime/GeometryCache/Private/GeometryCacheTrackTransformGroupAnimation.cpp b/Engine/Source/Runtime/GeometryCache/Private/GeometryCacheTrackTransformGroupAnimation.cpp
new file mode 100644
--- /dev/null
+++ b/Engine/Source/Runtime/GeometryCache/Private/GeometryCacheTrackTransformGroupAnimation.cpp
@@ -0,0 +1,43 @@
+// Copyright 1998-2015 Epic Games, Inc. All Rights Reserved.
+
+#include "GeometryCacheModulePrivatePCH.h"
+#include "GeometryCacheTrackTransformGroupAnimation.h"
+
+
+GEOMETRYCACHE_API UGeometryCacheTrack_TransformGroupAnimation::UGeometryCacheTrack_TransformGroupAnimation(const FObjectInitializer& ObjectInitializer /*= FObjectInitializer::Get()*/) : UGeometryCacheTrack(ObjectInitializer)
+{
+}
+
+SIZE_T UGeometryCacheTrack_TransformGroupAnimation::GetResourceSize(EResourceSizeMode::Type Mode)
+{
+	SIZE_T ResourceSize = 0;
+	ResourceSize += UGeometryCacheTrack::GetResourceSize(Mode);
+	ResourceSize += MeshData.GetResourceSize();
+
+	return ResourceSize;
+}
+
+void UGeometryCacheTrack_TransformGroupAnimation::Serialize(FArchive& Ar)
+{
+	UGeometryCacheTrack::Serialize(Ar);
+	Ar << MeshData;
+}
+
+const bool UGeometryCacheTrack_TransformGroupAnimation::UpdateMeshData(const float Time, const bool bLooping, int32& InOutMeshSampleIndex, FGeometryCacheMeshData*& OutMeshData)
+{
+	// If InOutMeshSampleIndex equals -1 (first creation) update the OutVertices and InOutMeshSampleIndex
+	if (InOutMeshSampleIndex == -1)
+	{
+		OutMeshData = &MeshData;
+		InOutMeshSampleIndex = 0;
+		return true;
+	}
+
+	return false;
+}
+
+void UGeometryCacheTrack_TransformGroupAnimation::SetMesh(const FGeometryCacheMeshData& NewMeshData)
+{
+	MeshData = NewMeshData;
+	NumMaterials = NewMeshData.BatchesInfo.Num();
+}
diff --git a/Engine/Source/Runtime/GeometryCache/Public/GeometryCacheModule.h b/Engine/Source/Runtime/GeometryCache/Public/GeometryCacheModule.h
--- a/Engine/Source/Runtime/GeometryCache/Public/GeometryCacheModule.h
+++ b/Engine/Source/Runtime/GeometryCache/Public/GeometryCacheModule.h
@@ -1,7 +1,6 @@
 // Copyright 1998-2015 Epic Games, Inc. All Rights Reserved.
 
 #pragma once
-
 #include "ModuleManager.h"
 
 /**
@@ -10,6 +9,8 @@
 class FGeometryCacheModule : public IModuleInterface
 {
 public:
+	FGeometryCacheModule() {}
+
 	virtual void StartupModule();
 	virtual void ShutdownModule();
 
@@ -33,16 +34,6 @@ public:
 	{
 		return FModuleManager::Get().IsModuleLoaded( "GeometryCache" );
 	}
-
-	/**
-	* OnObjectReimported, used to fix up rendering data for GeometryCacheComponent that's using the reimported GeometryCache
-	*
-	* @param InObject - Reimported object
-	* @return void
-	*/
-	void OnObjectReimported(UObject* InObject);
-	
 private:
-	FDelegateHandle OnAssetReimportDelegateHandle;
 };
 
diff --git a/Engine/Source/Runtime/GeometryCache/Public/GeometryCacheModulePublicPCH.h b/Engine/Source/Runtime/GeometryCache/Public/GeometryCacheModulePublicPCH.h
new file mode 100644
--- /dev/null
+++ b/Engine/Source/Runtime/GeometryCache/Public/GeometryCacheModulePublicPCH.h
@@ -0,0 +1 @@
+// Copyright 1998-2015 Epic Games, Inc. All Rights Reserved.
